Fix asan buffer overflow issue with message bus message in the session manager by making allocation/deallocation symmetrical

diff --git a/Engine/Source/Runtime/SessionServices/Private/SessionManager.cpp b/Engine/Source/Runtime/SessionServices/Private/SessionManager.cpp
--- a/Engine/Source/Runtime/SessionServices/Private/SessionManager.cpp
+++ b/Engine/Source/Runtime/SessionServices/Private/SessionManager.cpp
@@ -239,8 +239,8 @@ void FSessionManager::SendPing()
 
 	if (MessageEndpoint.IsValid())
 	{
-		MessageEndpoint->Publish(new FEngineServicePing(), EMessageScope::Network);
-		MessageEndpoint->Publish(new FSessionServicePing(FPlatformProcess::UserName(false)), EMessageScope::Network);
+		MessageEndpoint->Publish(FMessageEndpoint::MakeMessage<FEngineServicePing>(), EMessageScope::Network);
+		MessageEndpoint->Publish(FMessageEndpoint::MakeMessage<FSessionServicePing>(FPlatformProcess::UserName(false)), EMessageScope::Network);
 	}
 
 	LastPingTime = FDateTime::UtcNow();
