Move MessageEndpoint new -> MakeMessage hitting ASan new/free mis/match

diff --git a/Engine/Source/Runtime/SessionServices/Private/SessionInstanceInfo.cpp b/Engine/Source/Runtime/SessionServices/Private/SessionInstanceInfo.cpp
--- a/Engine/Source/Runtime/SessionServices/Private/SessionInstanceInfo.cpp
+++ b/Engine/Source/Runtime/SessionServices/Private/SessionInstanceInfo.cpp
@@ -52,7 +52,7 @@ void FSessionInstanceInfo::UpdateFromMessage(const FSessionServicePong& Message,
 
 	if (MessageEndpoint.IsValid() && (ApplicationAddress != Context->GetSender()))
 	{
-		MessageEndpoint->Send(new FSessionServiceLogSubscribe(), Context->GetSender());
+		MessageEndpoint->Send(FMessageEndpoint::MakeMessage<FSessionServiceLogSubscribe>(), Context->GetSender());
 	}
 
 	Authorized = Message.Authorized;
diff --git a/Engine/Source/Runtime/SessionServices/Private/SessionService.cpp b/Engine/Source/Runtime/SessionServices/Private/SessionService.cpp
--- a/Engine/Source/Runtime/SessionServices/Private/SessionService.cpp
+++ b/Engine/Source/Runtime/SessionServices/Private/SessionService.cpp
@@ -123,7 +123,7 @@ void FSessionService::SendPong(const TSharedRef<IMessageContext, ESPMode::Thread
 		return;
 	}
 
-	FSessionServicePong* Message = new FSessionServicePong();
+	FSessionServicePong* Message = FMessageEndpoint::MakeMessage<FSessionServicePong>();
 	{
 		Message->Authorized = FApp::IsAuthorizedUser(UserName);
 		Message->BuildDate = FApp::GetBuildDate();
