Initialize IsInSendLog to false

diff --git a/Engine/Source/Runtime/SessionServices/Private/SessionService.cpp b/Engine/Source/Runtime/SessionServices/Private/SessionService.cpp
--- a/Engine/Source/Runtime/SessionServices/Private/SessionService.cpp
+++ b/Engine/Source/Runtime/SessionServices/Private/SessionService.cpp
@@ -13,6 +13,7 @@
 
 FSessionService::FSessionService(const TSharedRef<IMessageBus, ESPMode::ThreadSafe>& InMessageBus)
 	: MessageBusPtr(InMessageBus)
+	, IsInSendLog(false)
 { }
 
 
