Allow the changing of the StatNamedEvents settings from the TraceControlWindow.
Add IsPaused flag to the TraceControl status message.
Add rediscovery of unregistered selected sessions in TraceController.

diff --git a/Engine/Source/Developer/TraceTools/Private/Widgets/STraceControlToolbar.cpp b/Engine/Source/Developer/TraceTools/Private/Widgets/STraceControlToolbar.cpp
--- a/Engine/Source/Developer/TraceTools/Private/Widgets/STraceControlToolbar.cpp
+++ b/Engine/Source/Developer/TraceTools/Private/Widgets/STraceControlToolbar.cpp
@@ -270,6 +270,7 @@ bool STraceControlToolbar::ToggleStatNamedEvents_IsChecked() const
 void STraceControlToolbar::ToggleStatNamedEvents_Execute()
 {
 	bAreStatNamedEventsEnabled = !bAreStatNamedEventsEnabled;
+	TraceController->SetStatNamedEventsEnabled(bAreStatNamedEventsEnabled);
 }
 
 FText STraceControlToolbar::GetTraceTargetLabelText() const
@@ -305,6 +306,8 @@ FSlateIcon STraceControlToolbar::GetTraceTargetIcon() const
 void STraceControlToolbar::OnTraceStatusUpdated(const FTraceStatus& InStatus, FTraceStatus::EUpdateType InUpdateType)
 {
 	bIsTracing = InStatus.bIsTracing;
+	bIsPaused = InStatus.bIsPaused;
+	bAreStatNamedEventsEnabled = InStatus.bAreStatNamedEventsEnabled;
 }
 
 } // namespace UE::TraceTools
diff --git a/Engine/Source/Runtime/Engine/Private/TraceService.cpp b/Engine/Source/Runtime/Engine/Private/TraceService.cpp
--- a/Engine/Source/Runtime/Engine/Private/TraceService.cpp
+++ b/Engine/Source/Runtime/Engine/Private/TraceService.cpp
@@ -30,6 +30,7 @@ private:
 	void OnResume(const FTraceControlResume& Message, const TSharedRef<IMessageContext>& Context);
 	void OnBookmark(const FTraceControlBookmark& Message, const TSharedRef<IMessageContext>& Context);
 	void OnScreenshot(const FTraceControlScreenshot& Message, const TSharedRef<IMessageContext>& Context);
+	void OnSetStatNamedEvents(const FTraceControlSetStatNamedEvents& Message, const TSharedRef<IMessageContext>& Context);
 	
 	static void FillTraceStatusMessage(FTraceControlStatus* Message);
 
@@ -56,6 +57,7 @@ FTraceServiceImpl::FTraceServiceImpl()
 		.Handling<FTraceControlResume>(this, &FTraceServiceImpl::OnResume)
 		.Handling<FTraceControlBookmark>(this, &FTraceServiceImpl::OnBookmark)
 		.Handling<FTraceControlScreenshot>(this, &FTraceServiceImpl::OnScreenshot)
+		.Handling<FTraceControlSetStatNamedEvents>(this, &FTraceServiceImpl::OnSetStatNamedEvents)
 		.Handling<FTraceControlStatusPing>(this, &FTraceServiceImpl::OnStatusPing)
 		.Handling<FTraceControlSettingsPing>(this, &FTraceServiceImpl::OnSettingsPing)
 		.Handling<FTraceControlChannelsPing>(this, &FTraceServiceImpl::OnChannelsPing);
@@ -79,6 +81,7 @@ FTraceServiceImpl::FTraceServiceImpl()
 	MessageEndpoint->Subscribe<FTraceControlResume>();
 	MessageEndpoint->Subscribe<FTraceControlBookmark>();
 	MessageEndpoint->Subscribe<FTraceControlScreenshot>();
+	MessageEndpoint->Subscribe<FTraceControlSetStatNamedEvents>();
 	
 }
 
@@ -97,7 +100,8 @@ void FTraceServiceImpl::FillTraceStatusMessage(FTraceControlStatus* Message)
 	Message->CacheAllocated = Stats.CacheAllocated;
 	Message->CacheUsed = Stats.CacheUsed;
 	Message->CacheWaste = Stats.CacheWaste;
-
+	Message->bAreStatNamedEventsEnabled = GCycleStatsShouldEmitNamedEvents > 0;
+	Message->bIsPaused = FTraceAuxiliary::IsPaused();
 }
 
 void FTraceServiceImpl::OnChannelSet(const FTraceControlChannelsSet& Message, const TSharedRef<IMessageContext>& Context)
@@ -168,6 +172,17 @@ void FTraceServiceImpl::OnScreenshot(const FTraceControlScreenshot& Message, con
 	FTraceScreenshot::RequestScreenshot(Message.Name, Message.bShowUI);
 }
 
+void FTraceServiceImpl::OnSetStatNamedEvents(const FTraceControlSetStatNamedEvents& Message, const TSharedRef<IMessageContext>& Context)
+{
+	if (Message.bEnabled && GCycleStatsShouldEmitNamedEvents == 0)
+	{
+		++GCycleStatsShouldEmitNamedEvents;
+	}
+	if (!Message.bEnabled && GCycleStatsShouldEmitNamedEvents > 0)
+	{
+		GCycleStatsShouldEmitNamedEvents = 0;
+	}
+}
 
 void FTraceServiceImpl::OnStatusPing(const FTraceControlStatusPing& Message, const TSharedRef<IMessageContext>& Context)
 {
@@ -230,7 +245,6 @@ void FTraceServiceImpl::OnSettingsPing(const FTraceControlSettingsPing& Message,
 	FTraceControlSettings* Response = FMessageEndpoint::MakeMessage<FTraceControlSettings>();
 	UE::Trace::FInitializeDesc const* InitDesc = FTraceAuxiliary::GetInitializeDesc();
 
-	Response->bStatNamedEvents = GCycleStatsShouldEmitNamedEvents > 0;
 	Response->bUseImportantCache = InitDesc->bUseImportantCache;
 	Response->bUseWorkerThread = InitDesc->bUseWorkerThread;
 	Response->TailSizeBytes = InitDesc->TailSizeBytes;
diff --git a/Engine/Source/Runtime/EngineMessages/Public/TraceControlMessages.h b/Engine/Source/Runtime/EngineMessages/Public/TraceControlMessages.h
--- a/Engine/Source/Runtime/EngineMessages/Public/TraceControlMessages.h
+++ b/Engine/Source/Runtime/EngineMessages/Public/TraceControlMessages.h
@@ -1,4 +1,4 @@
-﻿// Copyright Epic Games, Inc. All Rights Reserved.
+// Copyright Epic Games, Inc. All Rights Reserved.
 
 #pragma once
 
@@ -55,6 +55,12 @@ struct FTraceControlStatus
 
 	UPROPERTY(EditAnywhere, Category="Message")
 	uint32 CacheWaste = 0;
+
+	UPROPERTY(EditAnywhere, Category = "Message")
+	bool bAreStatNamedEventsEnabled = false;
+
+	UPROPERTY(EditAnywhere, Category = "Message")
+	bool bIsPaused = false;
 };
 
 /**
@@ -83,9 +89,6 @@ struct FTraceControlSettings
 
 	UPROPERTY(EditAnywhere, Category="Message")
 	uint32 TailSizeBytes = 0;
-	
-	UPROPERTY(EditAnywhere, Category="Message")
-	bool bStatNamedEvents = false;
 };
 
 /**
@@ -254,7 +257,7 @@ struct FTraceControlResume
 };
 
 /**
- * Start tracing to a host
+ * Trace a snapshot to the trace server.
  */
 USTRUCT()
 struct FTraceControlSnapshotSend
@@ -266,7 +269,7 @@ struct FTraceControlSnapshotSend
 };
 
 /**
- * Start tracing to a host
+ * Trace a snapshot to a file.
  */
 USTRUCT()
 struct FTraceControlSnapshotFile
@@ -278,7 +281,7 @@ struct FTraceControlSnapshotFile
 };
 
 /**
- * Start tracing to a host
+ * Trace a bookmark with the given label.
  */
 USTRUCT()
 struct FTraceControlBookmark
@@ -290,7 +293,7 @@ struct FTraceControlBookmark
 };
 
 /**
- * Start tracing to a host
+ * Trace a screenshot with the given label.
  */
 USTRUCT()
 struct FTraceControlScreenshot
@@ -303,3 +306,15 @@ struct FTraceControlScreenshot
 	UPROPERTY(EditAnywhere, Category="Message")
 	bool bShowUI = true;
 };
+
+/**
+ * Set the StatNamedEvents flag.
+ */
+USTRUCT()
+struct FTraceControlSetStatNamedEvents
+{
+	GENERATED_USTRUCT_BODY()
+
+	UPROPERTY(EditAnywhere, Category = "Message")
+	bool bEnabled = false;
+};
diff --git a/Engine/Source/Runtime/SessionServices/Private/TraceController.cpp b/Engine/Source/Runtime/SessionServices/Private/TraceController.cpp
--- a/Engine/Source/Runtime/SessionServices/Private/TraceController.cpp
+++ b/Engine/Source/Runtime/SessionServices/Private/TraceController.cpp
@@ -161,15 +161,28 @@ void FTraceController::Screenshot(FStringView Name, bool bShowUI)
 	SendToSelectedSessions(Message);
 }
 
+void FTraceController::SetStatNamedEventsEnabled(bool bEnabled)
+{
+	const auto Message = FMessageEndpoint::MakeMessage<FTraceControlSetStatNamedEvents>();
+	Message->bEnabled = bEnabled;
+
+	SendToSelectedSessions(Message);
+}
+
 void FTraceController::SendStatusUpdateRequest()
 {
-	MessageEndpoint->Publish(FMessageEndpoint::MakeMessage<FTraceControlStatusPing>());
+	// The discovery pong contains the status, so don't send the request twice.
+	if (RediscoverSelectedSession() == false)
+	{
+		MessageEndpoint->Publish(FMessageEndpoint::MakeMessage<FTraceControlStatusPing>());
+	}
 }
 
 void FTraceController::SendChannelUpdateRequest()
 {
 	FReadScopeLock _(InstancesLock);
 	
+	RediscoverSelectedSession();
 	for (const auto& Instance : Instances)
 	{
 		const auto Message = FMessageEndpoint::MakeMessage<FTraceControlChannelsPing>();
@@ -193,7 +206,7 @@ void FTraceController::OnNotification(const FMessageBusNotification& Event)
 		// one we have registered.
 		if(Instances.Remove(Event.RegistrationAddress) > 0)
 		{
-			InstanceToAddress = InstanceToAddress.FilterByPredicate([&](auto It) { return It.Value == Event.RegistrationAddress; } );
+			InstanceToAddress = InstanceToAddress.FilterByPredicate([&](auto It) { return It.Value != Event.RegistrationAddress; } );
 		}
 	}
 }
@@ -293,7 +306,6 @@ void FTraceController::OnSettings(const FTraceControlSettings& Message, const TS
 	if (const auto Status = Instances.Find(Context->GetSender()))
 	{
 		FTraceStatus::FSettings& Settings = Status->Settings;
-		Settings.bStatNamedEvents = Message.bStatNamedEvents;
 		Settings.bUseImportantCache = Message.bUseImportantCache;
 		Settings.bUseWorkerThread = Message.bUseWorkerThread;
 		Settings.TailSizeBytes = Message.TailSizeBytes;
@@ -337,12 +349,7 @@ void FTraceController::OnInstanceSelectionChanged(const TSharedPtr<ISessionInsta
 	}
 	else
 	{
-		// We haven't discovered this instance yet, send a discovery message specifically
-		// to that instance.
-		const auto Message = FMessageEndpoint::MakeMessage<FTraceControlDiscoveryPing>();
-		Message->SessionId = Instance->GetOwnerSession()->GetSessionId();
-		Message->InstanceId = Instance->GetInstanceId();
-		MessageEndpoint->Publish<FTraceControlDiscoveryPing>(Message);
+		SendDiscoveryPing(Instance);
 	}
 }
 
@@ -366,9 +373,40 @@ void FTraceController::UpdateStatus(const FTraceControlStatus& Message, FTraceSt
 	Status.Endpoint = Message.Endpoint;
 	Status.SessionGuid = Message.SessionGuid;
 	Status.TraceGuid = Message.TraceGuid;
+	Status.bIsPaused = Message.bIsPaused;
+	Status.bAreStatNamedEventsEnabled = Message.bAreStatNamedEventsEnabled;
 	Status.Stats.BytesSent = Message.BytesSent;
 	Status.Stats.BytesTraced = Message.BytesTraced;
 	Status.Stats.CacheAllocated = Message.CacheAllocated;
 	Status.Stats.CacheUsed = Message.CacheUsed;
 	Status.Stats.CacheWaste = Message.CacheWaste;
 }
+
+bool FTraceController::RediscoverSelectedSession()
+{
+	if (SelectedInstanceId.IsValid() && !InstanceToAddress.Contains(SelectedInstanceId))
+	{
+		auto& SelectedInstances = SessionManager->GetSelectedInstances();
+		for (auto& Instance : SelectedInstances)
+		{
+			if (Instance->GetInstanceId() == SelectedInstanceId)
+			{
+				SendDiscoveryPing(Instance);
+				return true;
+			}
+		}
+	}
+
+	return false;
+}
+
+void FTraceController::SendDiscoveryPing(const TSharedPtr<ISessionInstanceInfo>& Instance)
+{
+	// We haven't discovered this instance yet, send a discovery message specifically
+	// to that instance.
+
+	const auto Message = FMessageEndpoint::MakeMessage<FTraceControlDiscoveryPing>();
+	Message->SessionId = Instance->GetOwnerSession()->GetSessionId();
+	Message->InstanceId = Instance->GetInstanceId();
+	MessageEndpoint->Publish<FTraceControlDiscoveryPing>(Message);
+}
diff --git a/Engine/Source/Runtime/SessionServices/Private/TraceController.h b/Engine/Source/Runtime/SessionServices/Private/TraceController.h
--- a/Engine/Source/Runtime/SessionServices/Private/TraceController.h
+++ b/Engine/Source/Runtime/SessionServices/Private/TraceController.h
@@ -37,6 +37,7 @@ private:
 	virtual void Resume() override;
 	virtual void Bookmark(FStringView Label) override;
 	virtual void Screenshot(FStringView Name, bool bShowUI) override;
+	virtual void SetStatNamedEventsEnabled(bool bEnabled) override;
 
 	virtual void SendStatusUpdateRequest() override;
 	virtual void SendChannelUpdateRequest() override;
@@ -68,6 +69,14 @@ private:
 
 	TArray<FMessageAddress> GetSelectedSessionAddresses();
 
+	/* A selected instance can end up not discovered, either because the FTraceControlDiscoveryPong was lost
+	*  or because the selected session has been unregistered. Attempt to discover it again.
+	*  Returns true if a discovery ping was sent. 
+	*/
+	bool RediscoverSelectedSession();
+
+	void SendDiscoveryPing(const TSharedPtr<ISessionInstanceInfo>& Instance);
+
 public:
 
 private:
diff --git a/Engine/Source/Runtime/SessionServices/Public/ITraceController.h b/Engine/Source/Runtime/SessionServices/Public/ITraceController.h
--- a/Engine/Source/Runtime/SessionServices/Public/ITraceController.h
+++ b/Engine/Source/Runtime/SessionServices/Public/ITraceController.h
@@ -29,8 +29,6 @@ struct FTraceStatus
 		bool bUseImportantCache;
 		/** Size of tail buffer */
 		uint32 TailSizeBytes;
-		/** If stats are emitted as named events */
-		bool bStatNamedEvents;
 	};
 
 	struct FChannel
@@ -64,6 +62,10 @@ struct FTraceStatus
 	FGuid InstanceId;
 	/** If tracing is active */
 	bool bIsTracing = false;
+	/** If tracing is paused */
+	bool bIsPaused = false;
+	/** If stats are emitted as named events */
+	bool bAreStatNamedEventsEnabled = false;
 	/** Endpoint of active trace */
 	FString Endpoint;
 	/** Session identifier for the trace. */
@@ -162,6 +164,12 @@ public:
 	 */
 	virtual void Screenshot(FStringView Name, bool bShowUI) = 0;
 
+	/**
+	 * Set the StatNamedEvents flag.
+	 * @param bEnabled The value to assign to the StatNamedEvents flag.
+	 */
+	virtual void SetStatNamedEventsEnabled(bool bEnabled) = 0;
+
 	/**
 	 * Request update of the status from all sessions and instances.
 	 */
