Refactor TraceController and break out commands

diff --git a/Engine/Source/Developer/TraceTools/Private/Services/SessionTraceControllerFilterService.cpp b/Engine/Source/Developer/TraceTools/Private/Services/SessionTraceControllerFilterService.cpp
--- a/Engine/Source/Developer/TraceTools/Private/Services/SessionTraceControllerFilterService.cpp
+++ b/Engine/Source/Developer/TraceTools/Private/Services/SessionTraceControllerFilterService.cpp
@@ -13,14 +13,14 @@ FSessionTraceControllerFilterService::FSessionTraceControllerFilterService(TShar
 	FCoreDelegates::OnEndFrame.AddRaw(this, &FSessionTraceControllerFilterService::OnApplyChannelChanges);
 
 	TraceController = InTraceController;
-	TraceController->OnStatusReceived().AddRaw(this, &FSessionTraceControllerFilterService::OnTraceStatusUpdated);
+	TraceController->OnSelectedSessionStatusReceived().AddRaw(this, &FSessionTraceControllerFilterService::OnTraceStatusUpdated);
 }
 
 FSessionTraceControllerFilterService::~FSessionTraceControllerFilterService()
 {
 	FCoreDelegates::OnEndFrame.RemoveAll(this);
 
-	TraceController->OnStatusReceived().RemoveAll(this);
+	TraceController->OnSelectedSessionStatusReceived().RemoveAll(this);
 }
 
 void FSessionTraceControllerFilterService::GetRootObjects(TArray<FTraceObjectInfo>& OutObjects) const
@@ -74,7 +74,7 @@ void FSessionTraceControllerFilterService::DisableAllChannels()
 	}
 }
 
-void FSessionTraceControllerFilterService::OnTraceStatusUpdated(const FTraceStatus& InStatus, FTraceStatus::EUpdateType InUpdateType)
+void FSessionTraceControllerFilterService::OnTraceStatusUpdated(const FTraceStatus& InStatus, FTraceStatus::EUpdateType InUpdateType, ITraceControllerCommands& Commands)
 {
 	if (!TraceController->HasAvailableSelectedInstance())
 	{
@@ -109,7 +109,10 @@ void FSessionTraceControllerFilterService::OnApplyChannelChanges()
 {
 	if (FrameEnabledChannels.Num() || FrameDisabledChannels.Num())
 	{
-		TraceController->SetChannels(FrameEnabledChannels, FrameDisabledChannels);
+		TraceController->WithSelectedInstances([&](ITraceControllerCommands& Commands)
+		{
+			Commands.SetChannels(FrameEnabledChannels, FrameDisabledChannels);
+		});
 		FrameEnabledChannels.Empty();
 		FrameDisabledChannels.Empty();
 	}
diff --git a/Engine/Source/Developer/TraceTools/Private/Services/SessionTraceControllerFilterService.h b/Engine/Source/Developer/TraceTools/Private/Services/SessionTraceControllerFilterService.h
--- a/Engine/Source/Developer/TraceTools/Private/Services/SessionTraceControllerFilterService.h
+++ b/Engine/Source/Developer/TraceTools/Private/Services/SessionTraceControllerFilterService.h
@@ -41,7 +41,7 @@ protected:
 
 	void RetrieveAndStoreStartupChannels();
 
-	void OnTraceStatusUpdated(const FTraceStatus& InStatus, FTraceStatus::EUpdateType InUpdateType);
+	void OnTraceStatusUpdated(const FTraceStatus& InStatus, FTraceStatus::EUpdateType InUpdateType, ITraceControllerCommands& Commands);
 
 protected:
 	TSharedPtr<ITraceController> TraceController;
diff --git a/Engine/Source/Developer/TraceTools/Private/Widgets/STraceControlToolbar.cpp b/Engine/Source/Developer/TraceTools/Private/Widgets/STraceControlToolbar.cpp
--- a/Engine/Source/Developer/TraceTools/Private/Widgets/STraceControlToolbar.cpp
+++ b/Engine/Source/Developer/TraceTools/Private/Widgets/STraceControlToolbar.cpp
@@ -24,7 +24,7 @@ STraceControlToolbar::STraceControlToolbar()
 
 STraceControlToolbar::~STraceControlToolbar()
 {
-	TraceController->OnStatusReceived().Remove(OnStatusReceivedDelegate);
+	TraceController->OnSelectedSessionStatusReceived().Remove(OnStatusReceivedDelegate);
 }
 
 BEGIN_SLATE_FUNCTION_BUILD_OPTIMIZATION
@@ -175,14 +175,17 @@ bool STraceControlToolbar::StartTrace_CanExecute() const
 
 void STraceControlToolbar::StartTrace_Execute()
 {
-	if (TraceTarget == ETraceTarget::Server)
-	{
-		TraceController->Send(TraceHostAddr, TEXT(""));
-	}
-	else if (TraceTarget == ETraceTarget::File)
+	TraceController->WithSelectedInstances([&](ITraceControllerCommands& Commands)
 	{
-		TraceController->File(TEXT(""), TEXT(""));
-	}
+		if (TraceTarget == ETraceTarget::Server)
+		{
+			Commands.Send(TraceHostAddr, TEXT(""));
+		}
+		else if (TraceTarget == ETraceTarget::File)
+		{
+			Commands.File(TEXT(""), TEXT(""));
+		}
+	});
 	bIsTracing = true;
 }
 
@@ -193,7 +196,10 @@ bool STraceControlToolbar::StopTrace_CanExecute() const
 
 void STraceControlToolbar::StopTrace_Execute()
 {
-	TraceController->Stop();
+	TraceController->WithSelectedInstances([&](ITraceControllerCommands& Commands)
+	{
+		Commands.Stop();
+	});
 	bIsTracing = false;
 }
 
@@ -204,14 +210,17 @@ bool STraceControlToolbar::TraceSnapshot_CanExecute() const
 
 void STraceControlToolbar::TraceSnapshot_Execute()
 {
-	if (TraceTarget == ETraceTarget::Server)
-	{
-		TraceController->SnapshotSend(TraceHostAddr);
-	}
-	else if (TraceTarget == ETraceTarget::File)
+	TraceController->WithSelectedInstances([&](ITraceControllerCommands& Commands)
 	{
-		TraceController->SnapshotFile(TEXT(""));
-	}
+		if (TraceTarget == ETraceTarget::Server)
+		{
+			Commands.SnapshotSend(TraceHostAddr);
+		}
+		else if (TraceTarget == ETraceTarget::File)
+		{
+			Commands.SnapshotFile(TEXT(""));
+		}
+	});
 }
 
 bool STraceControlToolbar::PauseTrace_CanExecute() const
@@ -221,7 +230,10 @@ bool STraceControlToolbar::PauseTrace_CanExecute() const
 
 void STraceControlToolbar::PauseTrace_Execute()
 {
-	TraceController->Pause();
+	TraceController->WithSelectedInstances([&](ITraceControllerCommands& Commands)
+	{
+		Commands.Pause();
+	});
 	bIsPaused = true;
 }
 
@@ -232,7 +244,10 @@ bool STraceControlToolbar::ResumeTrace_CanExecute() const
 
 void STraceControlToolbar::ResumeTrace_Execute()
 {
-	TraceController->Resume();
+	TraceController->WithSelectedInstances([&](ITraceControllerCommands& Commands)
+	{
+		Commands.Resume();
+	});
 	bIsPaused = false;
 }
 
@@ -244,7 +259,10 @@ bool STraceControlToolbar::TraceBookmark_CanExecute() const
 void STraceControlToolbar::TraceBookmark_Execute()
 {
 	const FString BookmarkName = FDateTime::Now().ToString(TEXT("Bookmark_%Y%m%d_%H%M%S"));
-	TraceController->Bookmark(BookmarkName);
+	TraceController->WithSelectedInstances([&](ITraceControllerCommands& Commands)
+	{
+		Commands.Bookmark(BookmarkName);
+	});
 }
 
 bool STraceControlToolbar::TraceScreenshot_CanExecute() const
@@ -254,7 +272,10 @@ bool STraceControlToolbar::TraceScreenshot_CanExecute() const
 
 void STraceControlToolbar::TraceScreenshot_Execute()
 {
-	TraceController->Screenshot(TEXT(""), false);
+	TraceController->WithSelectedInstances([&](ITraceControllerCommands& Commands)
+	{
+		Commands.Screenshot(TEXT(""), false);
+	});
 }
 
 bool STraceControlToolbar::ToggleStatNamedEvents_CanExecute() const
@@ -270,7 +291,10 @@ bool STraceControlToolbar::ToggleStatNamedEvents_IsChecked() const
 void STraceControlToolbar::ToggleStatNamedEvents_Execute()
 {
 	bAreStatNamedEventsEnabled = !bAreStatNamedEventsEnabled;
-	TraceController->SetStatNamedEventsEnabled(bAreStatNamedEventsEnabled);
+	TraceController->WithSelectedInstances([&](ITraceControllerCommands& Commands)
+	{
+		Commands.SetStatNamedEventsEnabled(bAreStatNamedEventsEnabled);
+	});
 }
 
 FText STraceControlToolbar::GetTraceTargetLabelText() const
@@ -303,7 +327,7 @@ FSlateIcon STraceControlToolbar::GetTraceTargetIcon() const
 	return FSlateIcon(FTraceToolsStyle::GetStyleSetName(), "TraceControl.SetTraceTargetFile");
 }
 
-void STraceControlToolbar::OnTraceStatusUpdated(const FTraceStatus& InStatus, FTraceStatus::EUpdateType InUpdateType)
+void STraceControlToolbar::OnTraceStatusUpdated(const FTraceStatus& InStatus, FTraceStatus::EUpdateType InUpdateType, ITraceControllerCommands& Commands)
 {
 	bIsTracing = InStatus.bIsTracing;
 	bIsPaused = InStatus.bIsPaused;
diff --git a/Engine/Source/Developer/TraceTools/Private/Widgets/STraceControlToolbar.h b/Engine/Source/Developer/TraceTools/Private/Widgets/STraceControlToolbar.h
--- a/Engine/Source/Developer/TraceTools/Private/Widgets/STraceControlToolbar.h
+++ b/Engine/Source/Developer/TraceTools/Private/Widgets/STraceControlToolbar.h
@@ -44,7 +44,7 @@ private:
 	FText GetTraceTargetTooltipText() const;
 	FSlateIcon GetTraceTargetIcon() const;
 
-	void OnTraceStatusUpdated(const FTraceStatus& InStatus, FTraceStatus::EUpdateType InUpdateType);
+	void OnTraceStatusUpdated(const FTraceStatus& InStatus, FTraceStatus::EUpdateType InUpdateType, ITraceControllerCommands& Commands);
 	void InitializeSettings();
 
 	bool SetTraceTarget_CanExecute() const;
diff --git a/Engine/Source/Runtime/SessionServices/Private/SessionManager.cpp b/Engine/Source/Runtime/SessionServices/Private/SessionManager.cpp
--- a/Engine/Source/Runtime/SessionServices/Private/SessionManager.cpp
+++ b/Engine/Source/Runtime/SessionServices/Private/SessionManager.cpp
@@ -338,3 +338,21 @@ bool FSessionManager::HandleTicker(float DeltaTime)
 
 	return true;
 }
+
+TSharedPtr<ISessionInstanceInfo> FSessionManager::GetInstance(const FGuid& Id) const
+{
+	for (const auto& SessionPair : Sessions)
+	{
+		const auto& Session = SessionPair.Value;
+		TArray<TSharedPtr<ISessionInstanceInfo>> Instances;
+		Session->GetInstances(Instances);
+		for (const auto& Instance : Instances)
+		{
+			if (Instance->GetInstanceId() == Id)
+			{
+				return Instance;
+			}
+		}
+	}
+	return nullptr;
+}
diff --git a/Engine/Source/Runtime/SessionServices/Private/SessionManager.h b/Engine/Source/Runtime/SessionServices/Private/SessionManager.h
--- a/Engine/Source/Runtime/SessionServices/Private/SessionManager.h
+++ b/Engine/Source/Runtime/SessionServices/Private/SessionManager.h
@@ -112,6 +112,8 @@ private:
 	/** Callback for ticks from the ticker. */
 	bool HandleTicker(float DeltaTime);
 
+	virtual TSharedPtr<ISessionInstanceInfo> GetInstance(const FGuid& Id) const override;
+
 private:
 
 	/** The address of the automation controller to where we can forward any automation workers found. */
diff --git a/Engine/Source/Runtime/SessionServices/Private/TraceController.cpp b/Engine/Source/Runtime/SessionServices/Private/TraceController.cpp
--- a/Engine/Source/Runtime/SessionServices/Private/TraceController.cpp
+++ b/Engine/Source/Runtime/SessionServices/Private/TraceController.cpp
@@ -8,6 +8,7 @@
 #include "MessageEndpointBuilder.h"
 #include "TraceControlMessages.h"
 
+
 FTraceController::FTraceController(const TSharedRef<IMessageBus>& InMessageBus)
 	: MessageBus(InMessageBus)
 {
@@ -38,136 +39,6 @@ FTraceController::~FTraceController()
 	}
 }
 
-template<typename StringType>
-TArray<uint32> ResolveChannelNames(TConstArrayView<StringType> ChannelNames, const TMap<uint32, FTraceStatus::FChannel> Channels)
-{
-	TArray<uint32> Ids;
-	if (ChannelNames.Num() > 0)
-	{
-		for (const auto& [Id, Channel] : Channels)
-		{
-			if (ChannelNames.Contains(Channel.Name))
-			{
-				Ids.Add(Channel.Id);
-			}
-		}
-	}
-	return MoveTemp(Ids);
-}
-
-void FTraceController::SetChannels(TConstArrayView<FStringView> ChannelsToEnable, TConstArrayView<FStringView> ChannelsToDisable)
-{
-	FReadScopeLock _(InstancesLock);
-
-	// Since each session might have different notion about channel ids
-	// we need to send a specific message to each instance using their unique
-	// channel mappings.
-	for (const auto& Instance : SessionManager->GetSelectedInstances())
-	{
-		if (const auto Address = InstanceToAddress.Find(Instance->GetInstanceId()))
-		{
-			const auto& InstanceStatus = Instances[*Address];
-			const auto Message = FMessageEndpoint::MakeMessage<FTraceControlChannelsSet>();
-			Message->ChannelIdsToEnable = ResolveChannelNames(ChannelsToEnable, InstanceStatus.Channels);
-			Message->ChannelIdsToDisable = ResolveChannelNames(ChannelsToDisable, InstanceStatus.Channels);
-			MessageEndpoint->Send(Message, *Address);
-		}
-	}
-}
-
-void FTraceController::SetChannels(TConstArrayView<FString> ChannelsToEnable, TConstArrayView<FString> ChannelsToDisable)
-{
-	FReadScopeLock _(InstancesLock);
-	
-	for (const auto& Instance : SessionManager->GetSelectedInstances())
-	{
-		if (const auto Address = InstanceToAddress.Find(Instance->GetInstanceId()))
-		{
-			const auto& InstanceStatus = Instances[*Address];
-			const auto Message = FMessageEndpoint::MakeMessage<FTraceControlChannelsSet>();
-			Message->ChannelIdsToEnable = ResolveChannelNames(ChannelsToEnable, InstanceStatus.Channels);
-			Message->ChannelIdsToDisable = ResolveChannelNames(ChannelsToDisable, InstanceStatus.Channels);
-			MessageEndpoint->Send(Message, *Address);
-		}
-	}
-}
-
-void FTraceController::Send(FStringView Host, FStringView Channels, bool bExcludeTail)
-{
-	const auto Message = FMessageEndpoint::MakeMessage<FTraceControlSend>();
-	Message->Host = Host;
-	Message->Channels = Channels;
-	Message->bExcludeTail = bExcludeTail;
-
-	SendToSelectedSessions(Message);
-}
-
-void FTraceController::File(FStringView File, FStringView Channels, bool bExcludeTail, bool bTruncateFile)
-{
-	const auto Message = FMessageEndpoint::MakeMessage<FTraceControlFile>();
-	Message->File = File;
-	Message->Channels = Channels;
-	Message->bExcludeTail = bExcludeTail;
-	Message->bTruncateFile = bTruncateFile;
-
-	SendToSelectedSessions(Message);
-}
-
-void FTraceController::Stop()
-{
-	SendToSelectedSessions(FMessageEndpoint::MakeMessage<FTraceControlStop>());
-}
-
-void FTraceController::SnapshotSend(FStringView Host)
-{
-	const auto Message = FMessageEndpoint::MakeMessage<FTraceControlSnapshotSend>();
-	Message->Host = Host;
-
-	SendToSelectedSessions(Message);
-}
-
-void FTraceController::SnapshotFile(FStringView File)
-{
-	const auto Message = FMessageEndpoint::MakeMessage<FTraceControlSnapshotFile>();
-	Message->File = File;
-
-	SendToSelectedSessions(Message);
-}
-
-void FTraceController::Pause()
-{
-	SendToSelectedSessions(FMessageEndpoint::MakeMessage<FTraceControlPause>());
-}
-
-void FTraceController::Resume()
-{
-	SendToSelectedSessions(FMessageEndpoint::MakeMessage<FTraceControlResume>());
-}
-
-void FTraceController::Bookmark(FStringView Label)
-{
-	const auto Message = FMessageEndpoint::MakeMessage<FTraceControlBookmark>();
-	Message->Label = Label;
-	
-	SendToSelectedSessions(Message);
-}
-
-void FTraceController::Screenshot(FStringView Name, bool bShowUI)
-{
-	const auto Message = FMessageEndpoint::MakeMessage<FTraceControlScreenshot>();
-	Message->Name = Name;
-	Message->bShowUI = bShowUI;
-	
-	SendToSelectedSessions(Message);
-}
-
-void FTraceController::SetStatNamedEventsEnabled(bool bEnabled)
-{
-	const auto Message = FMessageEndpoint::MakeMessage<FTraceControlSetStatNamedEvents>();
-	Message->bEnabled = bEnabled;
-
-	SendToSelectedSessions(Message);
-}
 
 void FTraceController::SendStatusUpdateRequest()
 {
@@ -186,24 +57,26 @@ void FTraceController::SendChannelUpdateRequest()
 	for (const auto& Instance : Instances)
 	{
 		const auto Message = FMessageEndpoint::MakeMessage<FTraceControlChannelsPing>();
-		Message->KnownChannelCount = uint32(Instance.Value.Channels.Num());
+		Message->KnownChannelCount = uint32(Instance.Value.Status.Channels.Num());
 		MessageEndpoint->Send(Message, Instance.Key);
 	}
 }
 
 void FTraceController::SendSettingsUpdateRequest()
 {
-	MessageEndpoint->Publish(FMessageEndpoint::MakeMessage<FTraceControlSettingsPing>());
+	for (const auto& Instance : Instances)
+	{
+		MessageEndpoint->Send(FMessageEndpoint::MakeMessage<FTraceControlSettingsPing>(), Instance.Key);
+	}
 }
 
 void FTraceController::OnNotification(const FMessageBusNotification& Event)
 {
-	FWriteScopeLock _(InstancesLock);
-	
 	if (Event.NotificationType == EMessageBusNotification::Unregistered)
 	{
 		// Many endpoints may be removed from one instance, look for the
 		// one we have registered.
+		FWriteScopeLock _(InstancesLock);
 		if(Instances.Remove(Event.RegistrationAddress) > 0)
 		{
 			InstanceToAddress = InstanceToAddress.FilterByPredicate([&](auto It) { return It.Value != Event.RegistrationAddress; } );
@@ -214,19 +87,34 @@ void FTraceController::OnNotification(const FMessageBusNotification& Event)
 void FTraceController::OnDiscoveryResponse(const FTraceControlDiscovery& Message, const TSharedRef<IMessageContext>& Context)
 {
 	FWriteScopeLock _(InstancesLock);
+
+	// If the message bus is no longer available there is no point in registering
+	// more sessions.
+	TSharedPtr<IMessageBus> Bus = MessageBus.Pin();
+	if (!Bus)
+	{
+		return;
+	}
+
+	FTracingInstance* Instance = Instances.Find(Context->GetSender());
+	if (!Instance)
+	{
+		// Create a new instance with default status
+		Instance = &Instances.Emplace(Context->GetSender(), {Bus.ToSharedRef(), Context->GetSender()});
+		InstanceToAddress.Add(Message.InstanceId, Context->GetSender());
+	}
 	
-	FTraceStatus& Status = Instances.FindOrAdd(Context->GetSender());
-	InstanceToAddress.FindOrAdd(Message.InstanceId, Context->GetSender());
+	UpdateStatus(Message, Instance->Status);
 	
-	UpdateStatus(Message, Status);
 	// This is the application session id, which is not necessarily the
-	// same as trace session (may be overridden by commandline)
-	Status.SessionId = Message.SessionId;
-	Status.InstanceId = Message.InstanceId;
+	// same as trace session (maybe overridden by commandline)
+	Instance->Status.SessionId = Message.SessionId;
+	Instance->Status.InstanceId = Message.InstanceId;
 
-	if (SelectedInstanceId == Message.InstanceId)
+	StatusReceivedEvent.Broadcast(Instance->Status, FTraceStatus::EUpdateType::All, Instance->Commands);
+	if (SelectedInstanceIds.Contains(Message.InstanceId))
 	{
-		StatusReceivedEvent.Broadcast(Status, FTraceStatus::EUpdateType::All);
+		SelectedSessionStatusReceivedEvent.Broadcast(Instance->Status, FTraceStatus::EUpdateType::All, Instance->Commands);
 	}
 }
 
@@ -234,67 +122,75 @@ void FTraceController::OnStatus(const FTraceControlStatus& Message, const TShare
 {
 	FWriteScopeLock _(InstancesLock);
 	
-	if (FTraceStatus* Status = Instances.Find(Context->GetSender()))
+	if (FTracingInstance* Instance = Instances.Find(Context->GetSender()))
 	{
-		UpdateStatus(Message, *Status);
+		FTraceStatus& Status = Instance->Status;
+		UpdateStatus(Message, Status);
 		
-		if (SelectedInstanceId == Status->InstanceId)
+		StatusReceivedEvent.Broadcast(Status, FTraceStatus::EUpdateType::Status, Instance->Commands);
+		if (SelectedInstanceIds.Contains(Status.InstanceId))
 		{
-			StatusReceivedEvent.Broadcast(*Status, FTraceStatus::EUpdateType::Status);
+			SelectedSessionStatusReceivedEvent.Broadcast(Status, FTraceStatus::EUpdateType::Status, Instance->Commands);
 		}
 	}
-	
 }
 
 void FTraceController::OnChannelsDesc(const FTraceControlChannelsDesc& Message, const TSharedRef<IMessageContext>& Context)
 {
 	FWriteScopeLock _(InstancesLock);
 	
-	const auto Status = Instances.Find(Context->GetSender());
+	const auto Instance = Instances.Find(Context->GetSender());
 	
-	if (!Status)
+	if (!Instance)
 	{
 		return;
 	}
-	
+
+	FTraceStatus& Status = Instance->Status;
 	check(Message.Channels.Num() == Message.Ids.Num() && Message.Channels.Num() == Message.Descriptions.Num());
 	const int32 Count = Message.Channels.Num();
 	for(int32 Index = 0; Index < Count; ++Index)
 	{
 		const uint32 Id = Message.Ids[Index];
-		if (Status->Channels.Contains(Id))
+		if (Status.Channels.Contains(Id))
 		{
 			continue;
 		}
 
-		FTraceStatus::FChannel& NewChannel = Status->Channels.Add(Id);
+		FTraceStatus::FChannel& NewChannel = Status.Channels.Add(Id);
 		NewChannel.Name = Message.Channels[Index];
 		NewChannel.Description = Message.Descriptions[Index];
 		NewChannel.Id = Id;
 		NewChannel.bEnabled = false;
 		NewChannel.bReadOnly = Message.ReadOnlyIds.Contains(Id);
 		
-		if (SelectedInstanceId == Status->InstanceId)
+		StatusReceivedEvent.Broadcast(Status, FTraceStatus::EUpdateType::ChannelsDesc, Instance->Commands);
+		if (SelectedInstanceIds.Contains(Status.InstanceId))
 		{
-			StatusReceivedEvent.Broadcast(*Status, FTraceStatus::EUpdateType::ChannelsDesc);
+			SelectedSessionStatusReceivedEvent.Broadcast(Status, FTraceStatus::EUpdateType::ChannelsDesc, Instance->Commands);
 		}
 	}
+
+	// Allow the commands instance to update its list of channels
+	Instance->Commands.OnChannelsDesc(Message);
 }
 
 void FTraceController::OnChannelsStatus(const FTraceControlChannelsStatus& Message, const TSharedRef<IMessageContext>& Context)
 {
 	FWriteScopeLock _(InstancesLock);
 	
-	if (const auto Status = Instances.Find(Context->GetSender()))
+	if (const auto Instance = Instances.Find(Context->GetSender()))
 	{
-		for (auto& [Id, Channel] : Status->Channels)
+		FTraceStatus& Status = Instance->Status;
+		for (auto& [Id, Channel] : Status.Channels)
 		{
 			Channel.bEnabled = Message.EnabledIds.Contains(Id);
 		}
 		
-		if (SelectedInstanceId == Status->InstanceId)
+		StatusReceivedEvent.Broadcast(Status, FTraceStatus::EUpdateType::ChannelsStatus, Instance->Commands);
+		if (SelectedInstanceIds.Contains(Status.InstanceId))
 		{
-			StatusReceivedEvent.Broadcast(*Status, FTraceStatus::EUpdateType::ChannelsStatus);
+			SelectedSessionStatusReceivedEvent.Broadcast(Status, FTraceStatus::EUpdateType::ChannelsStatus, Instance->Commands);
 		}
 	}
 }
@@ -303,70 +199,62 @@ void FTraceController::OnSettings(const FTraceControlSettings& Message, const TS
 {
 	FWriteScopeLock _(InstancesLock);
 	
-	if (const auto Status = Instances.Find(Context->GetSender()))
+	if (const auto Instance = Instances.Find(Context->GetSender()))
 	{
-		FTraceStatus::FSettings& Settings = Status->Settings;
+		FTraceStatus& Status = Instance->Status;
+		FTraceStatus::FSettings& Settings = Status.Settings;
 		Settings.bUseImportantCache = Message.bUseImportantCache;
 		Settings.bUseWorkerThread = Message.bUseWorkerThread;
 		Settings.TailSizeBytes = Message.TailSizeBytes;
 		
-		if (SelectedInstanceId == Status->InstanceId)
+		StatusReceivedEvent.Broadcast(Status, FTraceStatus::EUpdateType::Settings, Instance->Commands);
+		if (SelectedInstanceIds.Contains(Status.InstanceId))
 		{
-			StatusReceivedEvent.Broadcast(*Status, FTraceStatus::EUpdateType::Settings);
+			SelectedSessionStatusReceivedEvent.Broadcast(Status, FTraceStatus::EUpdateType::Settings, Instance->Commands);
 		}
 	}
 }
 
 bool FTraceController::HasAvailableSelectedInstance()
 {
-	if (SelectedInstanceId.IsValid())
+	// Return true if at least one of the selected instances is known
+	for (const FGuid& SelectedInstanceId : SelectedInstanceIds)
 	{
-		return InstanceToAddress.Find(SelectedInstanceId) != nullptr;
+		if (InstanceToAddress.Contains(SelectedInstanceId))
+		{
+			return true;
+		}
 	}
-
 	return false;
 }
 
-void FTraceController::OnInstanceSelectionChanged(const TSharedPtr<ISessionInstanceInfo>& Instance, bool bSelected)
+void FTraceController::OnInstanceSelectionChanged(const TSharedPtr<ISessionInstanceInfo>& InstanceInfo, bool bSelected)
 {
 	FReadScopeLock _(InstancesLock);
 	
 	if (bSelected)
 	{
-		SelectedInstanceId = Instance->GetInstanceId();
+		SelectedInstanceIds.Add(InstanceInfo->GetInstanceId());
 	}
 	else
 	{
-		SelectedInstanceId.Invalidate();
+		SelectedInstanceIds.Remove(InstanceInfo->GetInstanceId());
 	}
 
-	if (const auto Address = InstanceToAddress.Find(Instance->GetInstanceId()))
+	if (const auto Address = InstanceToAddress.Find(InstanceInfo->GetInstanceId()))
 	{
-		if (const auto Status = Instances.Find(*Address))
+		if (const auto Instance = Instances.Find(*Address))
 		{
-			StatusReceivedEvent.Broadcast(*Status, FTraceStatus::EUpdateType::All);
+			FTraceStatus& Status = Instance->Status;
+			SelectedSessionStatusReceivedEvent.Broadcast(Status, FTraceStatus::EUpdateType::All, Instance->Commands);
 		}
 	}
 	else
 	{
-		SendDiscoveryPing(Instance);
+		SendDiscoveryPing(InstanceInfo);
 	}
 }
 
-TArray<FMessageAddress> FTraceController::GetSelectedSessionAddresses()
-{
-	TArray<FMessageAddress> Selected;
-	auto SelectedInstances = SessionManager->GetSelectedInstances();
-	for (const auto& Instance : SelectedInstances)
-	{
-		if (const auto Address = InstanceToAddress.Find(Instance->GetInstanceId()))
-		{
-			Selected.Add(*Address);
-		}
-	}
-	return MoveTemp(Selected);
-}
-
 void FTraceController::UpdateStatus(const FTraceControlStatus& Message, FTraceStatus& Status)
 {
 	Status.bIsTracing = !Message.Endpoint.IsEmpty();
@@ -382,31 +270,60 @@ void FTraceController::UpdateStatus(const FTraceControlStatus& Message, FTraceSt
 	Status.Stats.CacheWaste = Message.CacheWaste;
 }
 
-bool FTraceController::RediscoverSelectedSession()
+void FTraceController::WithSelectedInstances(FCallback Func)
 {
-	if (SelectedInstanceId.IsValid() && !InstanceToAddress.Contains(SelectedInstanceId))
+	FReadScopeLock _(InstancesLock);
+	for (auto& SelectedInstance : SessionManager->GetSelectedInstances())
 	{
-		auto& SelectedInstances = SessionManager->GetSelectedInstances();
-		for (auto& Instance : SelectedInstances)
+		if(const auto& InstanceAddress = InstanceToAddress.Find(SelectedInstance->GetInstanceId()))
 		{
-			if (Instance->GetInstanceId() == SelectedInstanceId)
+			if(auto* Instance = Instances.Find(*InstanceAddress))
 			{
-				SendDiscoveryPing(Instance);
-				return true;
+				Func(Instance->Commands);
 			}
 		}
 	}
+}
 
+void FTraceController::WithInstance(FGuid InstanceId, FCallback Func)
+{
+	FReadScopeLock _(InstancesLock);
+	if(const auto& InstanceAddress = InstanceToAddress.Find(InstanceId))
+	{
+		if(auto* Instance = Instances.Find(*InstanceAddress))
+		{
+			Func(Instance->Commands);
+		}
+	}
+}
+
+bool FTraceController::RediscoverSelectedSession() const
+{
+	auto& SelectedInstances = SessionManager->GetSelectedInstances();
+	for (auto& SelectedInstance : SelectedInstances)
+	{
+		if (!InstanceToAddress.Contains(SelectedInstance->GetInstanceId()))
+		{
+			SendDiscoveryPing(SelectedInstance);
+			return true;
+		}
+	}
 	return false;
 }
 
-void FTraceController::SendDiscoveryPing(const TSharedPtr<ISessionInstanceInfo>& Instance)
+void FTraceController::SendDiscoveryPing(const TSharedPtr<ISessionInstanceInfo>& Instance) const
 {
 	// We haven't discovered this instance yet, send a discovery message specifically
 	// to that instance.
-
 	const auto Message = FMessageEndpoint::MakeMessage<FTraceControlDiscoveryPing>();
 	Message->SessionId = Instance->GetOwnerSession()->GetSessionId();
 	Message->InstanceId = Instance->GetInstanceId();
 	MessageEndpoint->Publish<FTraceControlDiscoveryPing>(Message);
 }
+
+FTraceController::FTracingInstance::FTracingInstance(const TSharedRef<IMessageBus>& InMessageBus, FMessageAddress InService)
+	: Status()
+	, Commands(InMessageBus, InService)
+{
+}
+
diff --git a/Engine/Source/Runtime/SessionServices/Private/TraceController.h b/Engine/Source/Runtime/SessionServices/Private/TraceController.h
--- a/Engine/Source/Runtime/SessionServices/Private/TraceController.h
+++ b/Engine/Source/Runtime/SessionServices/Private/TraceController.h
@@ -6,6 +6,7 @@
 #include "ISessionManager.h"
 #include "ITraceController.h"
 #include "MessageEndpoint.h"
+#include "TraceControllerCommands.h"
 #include "TraceControlMessages.h"
 
 struct FTraceControlSettings;
@@ -25,25 +26,12 @@ public:
 	virtual ~FTraceController() override;
 
 private:
-	/* ITraceController interface */
-	virtual void SetChannels(TConstArrayView<FStringView> ChannelsToEnable, TConstArrayView<FStringView> ChannelsToDisable) override;
-	virtual void SetChannels(TConstArrayView<FString> ChannelsToEnable, TConstArrayView<FString> ChannelsToDisable) override;
-	virtual void Send(FStringView Host, FStringView Channels, bool bExcludeTail) override;
-	virtual void File(FStringView File, FStringView Channels, bool bExcludeTail, bool bTruncateFile) override;
-	virtual void Stop() override;
-	virtual void SnapshotSend(FStringView Host) override;
-	virtual void SnapshotFile(FStringView File) override;
-	virtual void Pause() override;
-	virtual void Resume() override;
-	virtual void Bookmark(FStringView Label) override;
-	virtual void Screenshot(FStringView Name, bool bShowUI) override;
-	virtual void SetStatNamedEventsEnabled(bool bEnabled) override;
-
 	virtual void SendStatusUpdateRequest() override;
 	virtual void SendChannelUpdateRequest() override;
 	virtual void SendSettingsUpdateRequest() override;
-
 	virtual bool HasAvailableSelectedInstance() override;
+	virtual void WithSelectedInstances(FCallback Func) override;
+	virtual void WithInstance(FGuid InstanceId, FCallback Func) override;
 	
 	DECLARE_DERIVED_EVENT(FTraceController, ITraceController::FStatusRecievedEvent, FStatusRecievedEvent);
 	virtual FStatusRecievedEvent& OnStatusReceived() override
@@ -51,6 +39,11 @@ private:
 		return StatusReceivedEvent;
 	}
 
+	virtual FStatusRecievedEvent& OnSelectedSessionStatusReceived() override
+	{
+		return SelectedSessionStatusReceivedEvent;
+	}
+
 	/* Message handlers */
 	void OnNotification(const FMessageBusNotification& MessageBusNotification);
 	void OnDiscoveryResponse(const FTraceControlDiscovery& Message, const TSharedRef<IMessageContext>& Context);
@@ -63,23 +56,30 @@ private:
 	/* Events from SessionManager handlers */
 	void OnInstanceSelectionChanged(const TSharedPtr<class ISessionInstanceInfo>&, bool);
 
-	/** Utility for sending a message to all selected sessions */
-	template <class MessageType>
-	void SendToSelectedSessions(MessageType* Message);
-
-	TArray<FMessageAddress> GetSelectedSessionAddresses();
-
 	/* A selected instance can end up not discovered, either because the FTraceControlDiscoveryPong was lost
 	*  or because the selected session has been unregistered. Attempt to discover it again.
 	*  Returns true if a discovery ping was sent. 
 	*/
-	bool RediscoverSelectedSession();
-
-	void SendDiscoveryPing(const TSharedPtr<ISessionInstanceInfo>& Instance);
+	bool RediscoverSelectedSession() const;
 
-public:
+	/* Send a discovery ping to a specific instance rather than broadcast a request. */
+	void SendDiscoveryPing(const TSharedPtr<ISessionInstanceInfo>& Instance) const;
 
 private:
+	
+	struct FTracingInstance
+	{
+		FTraceStatus Status;
+		FTraceControllerCommands Commands;
+
+		FTracingInstance(const TSharedRef<IMessageBus>& InMessageBus, FMessageAddress Service);
+		FTracingInstance() = delete;
+	};
+
+	/**
+	 * Needed to create command instances when new sessions are discovered. We don't need a ref counted
+	 * pointer to the message bus.
+	 */
 	TWeakPtr<IMessageBus> MessageBus;
 
 	/** Our own endpoint for messages */
@@ -91,36 +91,21 @@ private:
 	/** Address of the runtime endpoint for trace controls */
 	FMessageAddress TraceControlAddress;
 
+	/** Event for status updates on any session */
 	FStatusRecievedEvent StatusReceivedEvent;
 
+	/** Event for status updates on a selected session */
+	FStatusRecievedEvent SelectedSessionStatusReceivedEvent;
+
 	/** Lock to protect access to Instances list */
 	FRWLock InstancesLock;
 
-	/** Map of instances and their respective last reported status */
-	TMap<FMessageAddress, FTraceStatus> Instances;
+	/** Known instances with an active trace service */
+	TMap<FMessageAddress, FTracingInstance> Instances;
 
 	/** Secondary lookup from instance -> address */
 	TMap<FGuid, FMessageAddress> InstanceToAddress;
 
 	/** Currently selected instance */
-	FGuid SelectedInstanceId;
+	TSet<FGuid> SelectedInstanceIds;
 };
-
-template <typename MessageType>
-void FTraceController::SendToSelectedSessions(MessageType* Message)
-{
-	TArray<FMessageAddress> Recipients;
-	auto SelectedInstances = SessionManager->GetSelectedInstances();
-	for (const auto& Instance : SelectedInstances)
-	{
-		if (const auto Address = InstanceToAddress.Find(Instance->GetInstanceId()))
-		{
-			Recipients.Add(*Address);
-		}
-	}
-
-	if (!Recipients.IsEmpty())
-	{
-		MessageEndpoint->Send(Message, Recipients);
-	}
-}
diff --git a/Engine/Source/Runtime/SessionServices/Private/TraceControllerCommands.cpp b/Engine/Source/Runtime/SessionServices/Private/TraceControllerCommands.cpp
new file mode 100644
--- /dev/null
+++ b/Engine/Source/Runtime/SessionServices/Private/TraceControllerCommands.cpp
@@ -0,0 +1,163 @@
+﻿// Copyright Epic Games, Inc. All Rights Reserved.
+
+#include "TraceControllerCommands.h"
+
+#include "Hash/xxhash.h"
+#include "MessageEndpoint.h"
+#include "MessageEndpointBuilder.h"
+#include "TraceControlMessages.h"
+
+
+//------------------------------------------------------------------------------------------
+// Utilities
+//------------------------------------------------------------------------------------------
+
+uint64 HashName(FStringView Name)
+{
+	return FXxHash64::HashBuffer(Name.GetData(), Name.Len() * sizeof(TCHAR)).Hash;
+}
+
+template<typename StringType>
+TArray<uint32> ResolveChannelNames(TConstArrayView<StringType> ChannelNames, const TMap<uint64, uint32>& Channels)
+{
+	TArray<uint32> Ids;
+	for (const auto& ChannelName : ChannelNames)
+	{
+		const uint64 Hash = HashName(FStringView(ChannelName));
+		if (const auto FoundId = Channels.Find(Hash))
+		{
+			Ids.Add(*FoundId);
+		}
+	}
+	return MoveTemp(Ids);
+}
+
+//------------------------------------------------------------------------------------------
+// FTraceControllerCommands
+//------------------------------------------------------------------------------------------
+
+FTraceControllerCommands::FTraceControllerCommands(const TSharedPtr<IMessageBus>& InMessageBus, FMessageAddress Service)
+	: ServiceAddress(Service)
+{
+	TStringBuilder<128> EndpointName;
+	EndpointName << TEXT("FTraceControllerCommands_") << Service.ToString();
+	MessageEndpoint = FMessageEndpoint::Builder(EndpointName.ToString(), InMessageBus.ToSharedRef());
+}
+
+FTraceControllerCommands::~FTraceControllerCommands()
+{
+}
+
+void FTraceControllerCommands::SetChannels(TConstArrayView<FStringView> ChannelsToEnable, TConstArrayView<FStringView> ChannelsToDisable)
+{
+	const auto Message = FMessageEndpoint::MakeMessage<FTraceControlChannelsSet>();
+	Message->ChannelIdsToEnable = ResolveChannelNames(ChannelsToEnable, SettableChannels);
+	Message->ChannelIdsToDisable = ResolveChannelNames(ChannelsToDisable, SettableChannels);
+	MessageEndpoint->Send(Message, ServiceAddress);
+}
+
+void FTraceControllerCommands::SetChannels(TConstArrayView<FString> ChannelsToEnable, TConstArrayView<FString> ChannelsToDisable)
+{
+	const auto Message = FMessageEndpoint::MakeMessage<FTraceControlChannelsSet>();
+	Message->ChannelIdsToEnable = ResolveChannelNames(ChannelsToEnable, SettableChannels);
+	Message->ChannelIdsToDisable = ResolveChannelNames(ChannelsToDisable, SettableChannels);
+	MessageEndpoint->Send(Message, ServiceAddress);
+}
+
+void FTraceControllerCommands::Send(FStringView Host, FStringView Channels, bool bExcludeTail)
+{
+	const auto Message = FMessageEndpoint::MakeMessage<FTraceControlSend>();
+	Message->Host = Host;
+	Message->Channels = Channels;
+	Message->bExcludeTail = bExcludeTail;
+
+	MessageEndpoint->Send(Message, ServiceAddress);
+}
+
+void FTraceControllerCommands::File(FStringView File, FStringView Channels, bool bExcludeTail, bool bTruncateFile)
+{
+	const auto Message = FMessageEndpoint::MakeMessage<FTraceControlFile>();
+	Message->File = File;
+	Message->Channels = Channels;
+	Message->bExcludeTail = bExcludeTail;
+	Message->bTruncateFile = bTruncateFile;
+
+	MessageEndpoint->Send(Message, ServiceAddress);
+}
+
+void FTraceControllerCommands::Stop()
+{
+	MessageEndpoint->Send(FMessageEndpoint::MakeMessage<FTraceControlStop>(), ServiceAddress);
+}
+
+void FTraceControllerCommands::SnapshotSend(FStringView Host)
+{
+	const auto Message = FMessageEndpoint::MakeMessage<FTraceControlSnapshotSend>();
+	Message->Host = Host;
+
+	MessageEndpoint->Send(Message, ServiceAddress);
+}
+
+void FTraceControllerCommands::SnapshotFile(FStringView File)
+{
+	const auto Message = FMessageEndpoint::MakeMessage<FTraceControlSnapshotFile>();
+	Message->File = File;
+
+	MessageEndpoint->Send(Message, ServiceAddress);
+}
+
+void FTraceControllerCommands::Pause()
+{
+	MessageEndpoint->Send(FMessageEndpoint::MakeMessage<FTraceControlPause>(), ServiceAddress);
+}
+
+void FTraceControllerCommands::Resume()
+{
+	MessageEndpoint->Send(FMessageEndpoint::MakeMessage<FTraceControlResume>(), ServiceAddress);
+}
+
+void FTraceControllerCommands::Bookmark(FStringView Label)
+{
+	const auto Message = FMessageEndpoint::MakeMessage<FTraceControlBookmark>();
+	Message->Label = Label;
+	
+	MessageEndpoint->Send(Message, ServiceAddress);
+}
+
+void FTraceControllerCommands::Screenshot(FStringView Name, bool bShowUI)
+{
+	const auto Message = FMessageEndpoint::MakeMessage<FTraceControlScreenshot>();
+	Message->Name = Name;
+	Message->bShowUI = bShowUI;
+	
+	MessageEndpoint->Send(Message, ServiceAddress);
+}
+
+void FTraceControllerCommands::SetStatNamedEventsEnabled(bool bEnabled)
+{
+	const auto Message = FMessageEndpoint::MakeMessage<FTraceControlSetStatNamedEvents>();
+	Message->bEnabled = bEnabled;
+
+	MessageEndpoint->Send(Message, ServiceAddress);
+}
+
+void FTraceControllerCommands::OnChannelsDesc(const FTraceControlChannelsDesc& Message)
+{
+	const int32 ChannelCount = Message.Channels.Num();
+	if (ChannelCount != SettableChannels.Num())
+	{
+		for (int32 ChannelIdx = 0; ChannelIdx < ChannelCount; ++ChannelIdx)
+		{
+			uint32 Id = Message.Ids[ChannelIdx];
+			// Do not add read-only channels to our list
+			if (Message.ReadOnlyIds.Contains(Id))
+			{
+				continue;
+			}
+			const FString& ChannelName = Message.Channels[ChannelIdx];
+
+			const uint64 Hash = HashName(FStringView(ChannelName));
+			SettableChannels.FindOrAdd(Hash, Id);
+		}
+	}
+}
diff --git a/Engine/Source/Runtime/SessionServices/Private/TraceControllerCommands.h b/Engine/Source/Runtime/SessionServices/Private/TraceControllerCommands.h
new file mode 100644
--- /dev/null
+++ b/Engine/Source/Runtime/SessionServices/Private/TraceControllerCommands.h
@@ -0,0 +1,46 @@
+﻿// Copyright Epic Games, Inc. All Rights Reserved.
+
+#pragma once
+#include "IMessageContext.h"
+#include "ITraceControllerCommands.h"
+
+class IMessageBus;
+class FMessageEndpoint;
+struct FTraceControlChannelsDesc;
+
+/**
+ * Logic for sending commands to a specific remote trace service.
+ */
+class FTraceControllerCommands : public ITraceControllerCommands
+{
+public:
+	FTraceControllerCommands(const TSharedPtr<IMessageBus>& MessageBus, FMessageAddress Service);
+	virtual ~FTraceControllerCommands() override;
+
+	void OnChannelsDesc(const FTraceControlChannelsDesc& Message);
+	
+private:	
+	/* ITraceController interface */
+	virtual void SetChannels(TConstArrayView<FStringView> ChannelsToEnable, TConstArrayView<FStringView> ChannelsToDisable) override;
+	virtual void SetChannels(TConstArrayView<FString> ChannelsToEnable, TConstArrayView<FString> ChannelsToDisable) override;
+	virtual void Send(FStringView Host, FStringView Channels, bool bExcludeTail) override;
+	virtual void File(FStringView File, FStringView Channels, bool bExcludeTail, bool bTruncateFile) override;
+	virtual void Stop() override;
+	virtual void SnapshotSend(FStringView Host) override;
+	virtual void SnapshotFile(FStringView File) override;
+	virtual void Pause() override;
+	virtual void Resume() override;
+	virtual void Bookmark(FStringView Label) override;
+	virtual void Screenshot(FStringView Name, bool bShowUI) override;
+	virtual void SetStatNamedEventsEnabled(bool bEnabled) override;
+
+private:
+	/** Address of the remote trace service to control */
+	FMessageAddress ServiceAddress;
+	
+	/** Our own endpoint for messages */
+	TSharedPtr<FMessageEndpoint> MessageEndpoint;
+
+	/** List of channels (hash of their name) that is possible to set (not read-only). */
+	TMap<uint64, uint32> SettableChannels;
+};
diff --git a/Engine/Source/Runtime/SessionServices/Public/ISessionManager.h b/Engine/Source/Runtime/SessionServices/Public/ISessionManager.h
--- a/Engine/Source/Runtime/SessionServices/Public/ISessionManager.h
+++ b/Engine/Source/Runtime/SessionServices/Public/ISessionManager.h
@@ -42,6 +42,15 @@ public:
 	 */
 	virtual void GetSessions( TArray<TSharedPtr<ISessionInfo>>& OutSessions ) const = 0;
 
+
+	/**
+	 * Attempts to retrieve an instance by it's unique instance id.
+	 * 
+	 * @param Id Instance id
+	 * @return Instance if it exist, otherwise a nullptr
+	 */
+	virtual TSharedPtr<ISessionInstanceInfo> GetInstance(const FGuid& Id) const = 0;
+
 	/**
 	 * Checks whether the given instance is currently selected.
 	 *
diff --git a/Engine/Source/Runtime/SessionServices/Public/ITraceController.h b/Engine/Source/Runtime/SessionServices/Public/ITraceController.h
--- a/Engine/Source/Runtime/SessionServices/Public/ITraceController.h
+++ b/Engine/Source/Runtime/SessionServices/Public/ITraceController.h
@@ -5,8 +5,13 @@
 #include "Containers/Map.h"
 #include "Containers/UnrealString.h"
 #include "Delegates/Delegate.h"
+#include "ITraceControllerCommands.h"
 #include "Misc/Guid.h"
 
+
+/**
+ * Describes the state of a single remote instance with Trace.
+ */
 struct FTraceStatus
 {
 	/**
@@ -33,10 +38,16 @@ struct FTraceStatus
 
 	struct FChannel
 	{
+		/** Channel name */
 		FString Name;
+		/** Description string */
 		FString Description;
+		/** Remote system id. This may be unique per machine. */
 		uint32 Id;
+		/** If channel is currently enabled */
 		bool bEnabled;
+		/** If it's possible to toggle channel. Read only channels
+		 * must be set on process start using command line arguments. */
 		bool bReadOnly;
 	};
 
@@ -82,8 +93,31 @@ struct FTraceStatus
 
 ENUM_CLASS_FLAGS(FTraceStatus::EUpdateType);
 
+
 /**
- * Interface to control other sessions tracing.
+ * Interface for monitoring known trace service instances. ITraceController is a singleton which tracks
+ * all known instances and provides callback for status updates and interfaces to issue commands to one
+ * or more remote instances.
+ *
+ * Example usage:
+ * ```
+ * // I've found the session that we want to control
+ * FGuid SessionId = ...;
+ * // Get the controller
+ * TSharedPtr<ITraceController> TraceController = SessionServicesModule.GetTraceController();
+ * 
+ * // To control just a single instance use the session id.
+ * TraceController->WithInstance(SessionId, [](ITraceControllerCommands& Commands){
+ *		Commands.Send("localhost", "audio,audiomixer,bookmark,log");
+ *		Commands.Bookmark("My remote bookmark");
+ * });
+ *
+ * // We can also use the "selection" feature in session manager
+ * TraceController->WithSelectedInstances([](ITraceControllerCommands& Commands){
+ *		Commands.SnapshotSend("localhost");
+ * });
+ * ```
+ * 
  */
 class ITraceController
 {
@@ -91,85 +125,6 @@ public:
 	
 	virtual ~ITraceController() = default;
 
-	/**
-	 * Enables or disables channels by name
-	 * @param ChannelsToEnable  List of channels to enable
-	 * @param ChannelsToDisable List of channels to disable
-	 */
-	virtual void SetChannels(TConstArrayView<FStringView> ChannelsToEnable, TConstArrayView<FStringView> ChannelsToDisable) = 0;
-
-	/**
-	 * Enables or disables channels by name
-	 * @param ChannelsToEnable  List of channels to enable
-	 * @param ChannelsToDisable List of channels to disable
-	 */
-	virtual void SetChannels(TConstArrayView<FString> ChannelsToEnable, TConstArrayView<FString> ChannelsToDisable) = 0;
-	
-	/**
-	 * Start a trace on selected instances to the provided host, using a set of channels.
-	 * @param Host Host to send the trace to
-	 * @param Channels Comma separated list of channels to enable
-	 * @param bExcludeTail If the tail (circular buffer of recent events) should be included
-	 */
-	virtual void Send(FStringView Host, FStringView Channels, bool bExcludeTail = false) = 0;
-
-	/**
-	 * Start a trace on selected instances to a file on the instance, using a set of channels.
-	 * @param File Path on the instance. ".utrace" will be appended
-	 * @param Channels Comma separated list of channels to enable
-	 * @param bExcludeTail If the tail (circular buffer of recent events) should be included
-	 * @param bTruncateFile If the file should be truncated (if already exists)
-	 */
-	virtual void File(FStringView File, FStringView Channels, bool bExcludeTail = false, bool bTruncateFile = false) = 0;
-
-	/**
-	 * On selected instances, make a snapshot of the tail (circular buffer of recent events)
-	 * and send to the provided host.
-	 * @param Host Host to send the trace to
-	 */
-	virtual void SnapshotSend(FStringView Host) = 0;
-	
-	/**
-	 * On selected instances, make a snapshot of the tail (circular buffer of recent events)
-	 * and save to a file.
-	 * @param File Path on the instance. ".utrace" will be appended
-	 */
-	virtual void SnapshotFile(FStringView File) = 0;
-
-	/**
-	 * Pause tracing by muting all (non-readonly) channels.
-	 */
-	virtual void Pause() = 0;
-
-	/**
-	 * Resume tracing (from paused) by enabling the previously enabled channels.
-	 */
-	virtual void Resume() = 0;
-
-	/**
-	 * Stop active trace.
-	 */
-	virtual void Stop() = 0;
-
-	/**
-	 * Insert bookmark into the trace.
-	 * @param Label Label of bookmark
-	 */
-	virtual void Bookmark(FStringView Label) = 0;
-
-	/**
-	 * Insert screenshot into the trace.
-	 * @param Name Name of the screenshot
-	 * @param bShowUI If the UI should be visible in the image
-	 */
-	virtual void Screenshot(FStringView Name, bool bShowUI) = 0;
-
-	/**
-	 * Set the StatNamedEvents flag.
-	 * @param bEnabled The value to assign to the StatNamedEvents flag.
-	 */
-	virtual void SetStatNamedEventsEnabled(bool bEnabled) = 0;
-
 	/**
 	 * Request update of the status from all sessions and instances.
 	 */
@@ -188,17 +143,39 @@ public:
 	/**
 	 * Generic event for updates of status
 	 */
-	DECLARE_EVENT_TwoParams(ITraceController, FStatusRecievedEvent, const FTraceStatus&, FTraceStatus::EUpdateType);
+	DECLARE_EVENT_ThreeParams(ITraceController, FStatusRecievedEvent, const FTraceStatus&, FTraceStatus::EUpdateType, ITraceControllerCommands&);
 
 	/**
-	 * Event triggered whenever the FTraceStatus for the selected instance (in session manager)
-	 * is updated. A reference to the status and what has changed is provided.
+	 * Event triggered whenever status is updated for any known session. 
+	 * A reference to the status and what has changed is provided along with a 
+	 * structure to issue commands.
 	 */
 	virtual FStatusRecievedEvent& OnStatusReceived() = 0;
 
+	/**
+	 * Event triggered whenever the status for the selected instance (in session manager)
+	 * is updated. A reference to the status and what has changed is provided along with a 
+	 * structure to issue commands.
+	 */
+	virtual FStatusRecievedEvent& OnSelectedSessionStatusReceived() = 0;
+
 	/**
 	 * Return true if a selected instance exists and has been discovered.
 	 */
 	virtual bool HasAvailableSelectedInstance() = 0;
-	
+
+
+	typedef TFunction<void(ITraceControllerCommands&)> FCallback;
+	/**
+	 * Execute a function for each instance selected in the session manager.
+	 * @param Func Functor to execute
+	 */
+	virtual void WithSelectedInstances(FCallback Func) = 0;
+
+	/**
+	 * Execute a function on a specific session.
+	 * @param InstanceId InstanceId to issue command to
+	 * @param Func Functor to execute
+	 */
+	virtual void WithInstance(FGuid InstanceId, FCallback Func) = 0;
 };
diff --git a/Engine/Source/Runtime/SessionServices/Public/ITraceControllerCommands.h b/Engine/Source/Runtime/SessionServices/Public/ITraceControllerCommands.h
new file mode 100644
--- /dev/null
+++ b/Engine/Source/Runtime/SessionServices/Public/ITraceControllerCommands.h
@@ -0,0 +1,96 @@
+﻿// Copyright Epic Games, Inc. All Rights Reserved.
+
+#pragma once
+
+#include "Containers/ContainersFwd.h"
+#include "Containers/StringFwd.h"
+
+/**
+ * Interface to send commands to a single remote trace service.
+ */
+class ITraceControllerCommands
+{
+public:
+
+	virtual ~ITraceControllerCommands() {}
+	
+	/**
+	 * Enables or disables channels by name
+	 * @param ChannelsToEnable  List of channels to enable
+	 * @param ChannelsToDisable List of channels to disable
+	 */
+	virtual void SetChannels(TConstArrayView<FStringView> ChannelsToEnable, TConstArrayView<FStringView> ChannelsToDisable) = 0;
+
+	/**
+	 * Enables or disables channels by name
+	 * @param ChannelsToEnable  List of channels to enable
+	 * @param ChannelsToDisable List of channels to disable
+	 */
+	virtual void SetChannels(TConstArrayView<FString> ChannelsToEnable, TConstArrayView<FString> ChannelsToDisable) = 0;
+	
+	/**
+	 * Start a trace on selected instances to the provided host, using a set of channels.
+	 * @param Host Host to send the trace to
+	 * @param Channels Comma separated list of channels to enable
+	 * @param bExcludeTail If the tail (circular buffer of recent events) should be included
+	 */
+	virtual void Send(FStringView Host, FStringView Channels, bool bExcludeTail = false) = 0;
+
+	/**
+	 * Start a trace on selected instances to a file on the instance, using a set of channels.
+	 * @param File Path on the instance. ".utrace" will be appended
+	 * @param Channels Comma separated list of channels to enable
+	 * @param bExcludeTail If the tail (circular buffer of recent events) should be included
+	 * @param bTruncateFile If the file should be truncated (if already exists)
+	 */
+	virtual void File(FStringView File, FStringView Channels, bool bExcludeTail = false, bool bTruncateFile = false) = 0;
+
+	/**
+	 * On selected instances, make a snapshot of the tail (circular buffer of recent events)
+	 * and send to the provided host.
+	 * @param Host Host to send the trace to
+	 */
+	virtual void SnapshotSend(FStringView Host) = 0;
+	
+	/**
+	 * On selected instances, make a snapshot of the tail (circular buffer of recent events)
+	 * and save to a file.
+	 * @param File Path on the instance. ".utrace" will be appended
+	 */
+	virtual void SnapshotFile(FStringView File) = 0;
+
+	/**
+	 * Pause tracing by muting all (non-readonly) channels.
+	 */
+	virtual void Pause() = 0;
+
+	/**
+	 * Resume tracing (from paused) by enabling the previously enabled channels.
+	 */
+	virtual void Resume() = 0;
+
+	/**
+	 * Stop active trace.
+	 */
+	virtual void Stop() = 0;
+
+	/**
+	 * Insert bookmark into the trace.
+	 * @param Label Label of bookmark
+	 */
+	virtual void Bookmark(FStringView Label) = 0;
+
+	/**
+	 * Insert screenshot into the trace.
+	 * @param Name Name of the screenshot
+	 * @param bShowUI If the UI should be visible in the image
+	 */
+	virtual void Screenshot(FStringView Name, bool bShowUI) = 0;
+
+	/**
+	 * Set the StatNamedEvents flag.
+	 * @param bEnabled The value to assign to the StatNamedEvents flag.
+	 */
+	virtual void SetStatNamedEventsEnabled(bool bEnabled) = 0;
+};
+
