Fix channel case insesitivity in trace controller

diff --git a/Engine/Source/Runtime/SessionServices/Private/TraceControllerCommands.cpp b/Engine/Source/Runtime/SessionServices/Private/TraceControllerCommands.cpp
--- a/Engine/Source/Runtime/SessionServices/Private/TraceControllerCommands.cpp
+++ b/Engine/Source/Runtime/SessionServices/Private/TraceControllerCommands.cpp
@@ -14,7 +14,12 @@
 
 uint64 HashName(FStringView Name)
 {
-	return FXxHash64::HashBuffer(Name.GetData(), Name.Len() * sizeof(TCHAR)).Hash;
+	// Strip plurals and convert to upper case.
+	TStringBuilder<48> TransformedBuffer;
+	TransformedBuffer << (Name.EndsWith('s') ? Name.LeftChop(1) : Name);
+	FCString::Strupr(TransformedBuffer.GetData(), TransformedBuffer.Len());
+	
+	return FXxHash64::HashBuffer(TransformedBuffer.GetData(), TransformedBuffer.Len() * sizeof(TCHAR)).Hash;
 }
 
 template<typename StringType>
