Added missing field in TraceController stats

diff --git a/Engine/Source/Runtime/SessionServices/Private/TraceController.cpp b/Engine/Source/Runtime/SessionServices/Private/TraceController.cpp
--- a/Engine/Source/Runtime/SessionServices/Private/TraceController.cpp
+++ b/Engine/Source/Runtime/SessionServices/Private/TraceController.cpp
@@ -280,6 +280,7 @@ void FTraceController::UpdateStatus(const FTraceControlStatus& Message, FTraceSt
 	Status.bIsPaused = Message.bIsPaused;
 	Status.bAreStatNamedEventsEnabled = Message.bAreStatNamedEventsEnabled;
 	Status.Stats.BytesSent = Message.BytesSent;
+	Status.Stats.MemoryUsed = Message.MemoryUsed;
 	Status.Stats.BytesTraced = Message.BytesTraced;
 	Status.Stats.CacheAllocated = Message.CacheAllocated;
 	Status.Stats.CacheUsed = Message.CacheUsed;
