Fix issue in the trace control window with the trace connection status not updating when the connection to the trace server is lost.

diff --git a/Engine/Source/Runtime/Engine/Private/TraceService.cpp b/Engine/Source/Runtime/Engine/Private/TraceService.cpp
--- a/Engine/Source/Runtime/Engine/Private/TraceService.cpp
+++ b/Engine/Source/Runtime/Engine/Private/TraceService.cpp
@@ -97,7 +97,7 @@ void FTraceServiceImpl::FillTraceStatusMessage(FTraceControlStatus* Message)
 {
 	// Get the current endpoint and ids
 	Message->Endpoint = FTraceAuxiliary::GetTraceDestinationString();
-	FTraceAuxiliary::IsConnected(Message->SessionGuid, Message->TraceGuid);
+	Message->bIsTracing = FTraceAuxiliary::IsConnected(Message->SessionGuid, Message->TraceGuid);
 	
 	// For stats we can query TraceLog directly.
 	UE::Trace::FStatistics Stats;
diff --git a/Engine/Source/Runtime/EngineMessages/Public/TraceControlMessages.h b/Engine/Source/Runtime/EngineMessages/Public/TraceControlMessages.h
--- a/Engine/Source/Runtime/EngineMessages/Public/TraceControlMessages.h
+++ b/Engine/Source/Runtime/EngineMessages/Public/TraceControlMessages.h
@@ -61,6 +61,9 @@ struct FTraceControlStatus
 
 	UPROPERTY(EditAnywhere, Category = "Message")
 	bool bIsPaused = false;
+
+	UPROPERTY(EditAnywhere, Category = "Message")
+	bool bIsTracing = false;
 };
 
 /**
diff --git a/Engine/Source/Runtime/SessionServices/Private/TraceController.cpp b/Engine/Source/Runtime/SessionServices/Private/TraceController.cpp
--- a/Engine/Source/Runtime/SessionServices/Private/TraceController.cpp
+++ b/Engine/Source/Runtime/SessionServices/Private/TraceController.cpp
@@ -273,7 +273,7 @@ void FTraceController::OnInstanceSelectionChanged(const TSharedPtr<ISessionInsta
 
 void FTraceController::UpdateStatus(const FTraceControlStatus& Message, FTraceStatus& Status)
 {
-	Status.bIsTracing = !Message.Endpoint.IsEmpty();
+	Status.bIsTracing = Message.bIsTracing;
 	Status.Endpoint = Message.Endpoint;
 	Status.SessionGuid = Message.SessionGuid;
 	Status.TraceGuid = Message.TraceGuid;
