Added the TraceControl window in Insights Session Browser.
Refactored TraceControl widgets to enable them to connect to a specific InstanceId.
Added auto detect selected session functionality to STraceControl.
Removed responsability from TraceController to handle selected session management.
Added InstanceId to the Diagnostics.SessionInfo2 event.

diff --git a/Engine/Source/Developer/TraceInsights/Private/Insights/Common/InsightsMenuBuilder.cpp b/Engine/Source/Developer/TraceInsights/Private/Insights/Common/InsightsMenuBuilder.cpp
--- a/Engine/Source/Developer/TraceInsights/Private/Insights/Common/InsightsMenuBuilder.cpp
+++ b/Engine/Source/Developer/TraceInsights/Private/Insights/Common/InsightsMenuBuilder.cpp
@@ -67,7 +67,9 @@ void FInsightsMenuBuilder::PopulateMenu(FMenuBuilder& MenuBuilder)
 		LOCTEXT("ImportTable_ToolTip", "Import CSV or TSV data from a file to an Insights Table."),
 		FSlateIcon(FInsightsStyle::GetStyleSetName(), "Icons.ImportTable"),
 		FUIAction(FExecuteAction::CreateLambda([] { Insights::FTableImportTool::Get()->StartImportProcess(); })));
+
 	MenuBuilder.AddSeparator();
+
 	MenuBuilder.AddMenuEntry(
 		LOCTEXT("OpenSessionBrowser", "Session Browser"),
 		LOCTEXT("OpenSessionBrowser_ToolTip", "Opens the Unreal Insights Session Browser window."),
@@ -80,6 +82,13 @@ void FInsightsMenuBuilder::PopulateMenu(FMenuBuilder& MenuBuilder)
 		false,
 		FSlateIcon(FAppStyle::Get().GetStyleSetName(), "Icons.FolderOpen")
 	);
+
+	MenuBuilder.AddMenuEntry(
+		LOCTEXT("OpenTraceControl", "Live Trace Control"),
+		LOCTEXT("OpenTraceControl_ToolTip", "Opens the Trace Control window."),
+		FSlateIcon(FInsightsStyle::GetStyleSetName(), "Icons.TraceControl"),
+		FUIAction(FExecuteAction::CreateLambda([] { FInsightsManager::Get()->OpenTraceControlWindow(); })));
+
 	MenuBuilder.AddMenuEntry(
 		LOCTEXT("AutoOpenLiveTrace", "Auto Open Live Trace"),
 		LOCTEXT("AutoOpenLiveTrace_ToolTip", "If enabled, the analysis starts automatically for each new live trace session, replacing the current analysis session."),
@@ -91,6 +100,7 @@ void FInsightsMenuBuilder::PopulateMenu(FMenuBuilder& MenuBuilder)
 		),
 		NAME_None,
 		EUserInterfaceActionType::ToggleButton);
+
 	MenuBuilder.EndSection();
 
 	FGlobalTabmanager::Get()->PopulateLocalTabSpawnerMenu(MenuBuilder);
diff --git a/Engine/Source/Developer/TraceInsights/Private/Insights/InsightsManager.cpp b/Engine/Source/Developer/TraceInsights/Private/Insights/InsightsManager.cpp
--- a/Engine/Source/Developer/TraceInsights/Private/Insights/InsightsManager.cpp
+++ b/Engine/Source/Developer/TraceInsights/Private/Insights/InsightsManager.cpp
@@ -4,8 +4,10 @@
 
 #include "DesktopPlatformModule.h"
 #include "Framework/Application/SlateApplication.h"
+#include "Framework/Docking/TabManager.h"
 #include "HAL/PlatformMemory.h"
 #include "HAL/PlatformProcess.h"
+#include "ISessionServicesModule.h"
 #include "ISourceCodeAccessModule.h"
 #include "ISourceCodeAccessor.h"
 #include "Logging/MessageLog.h"
@@ -18,6 +20,9 @@
 #include "WorkspaceMenuStructure.h"
 #include "WorkspaceMenuStructureModule.h"
 
+// TraceTools
+#include "TraceTools/Interfaces/ITraceToolsModule.h"
+
 // TraceAnalysis
 #include "Trace/StoreClient.h"
 
@@ -59,6 +64,7 @@ const FName FInsightsManagerTabs::NetworkingProfilerTabId(TEXT("NetworkingProfil
 
 const FName FInsightsManagerTabs::AutomationWindowTabId(TEXT("AutomationWindow"));
 const FName FInsightsManagerTabs::MessageLogTabId(TEXT("MessageLog"));
+const FName FInsightsManagerTabs::TraceControlTabId(TEXT("TraceControl"));
 
 namespace UE::Insights
 {
@@ -126,6 +132,8 @@ void FInsightsManager::Initialize(IUnrealInsightsModule& InsightsModule)
 	MessageLogModule.RegisterLogListing(AnalysisLogListingName, LOCTEXT("TraceAnalysis", "Trace Analysis"));
 	MessageLogModule.EnableMessageLogDisplay(true);
 
+	RegisterTraceControlTab();
+
 	// Register tick functions.
 	OnTick = FTickerDelegate::CreateSP(this, &FInsightsManager::Tick);
 	OnTickHandle = FTSTicker::GetCoreTicker().AddTicker(OnTick, 0.0f);
@@ -312,8 +320,15 @@ bool FInsightsManager::Tick(float DeltaTime)
 			if (DiagnosticsProvider && DiagnosticsProvider->IsSessionInfoAvailable())
 			{
 				bIsSessionInfoSet = true;
+				InstanceId = DiagnosticsProvider->GetSessionInfo().InstanceId;
+				TSharedPtr<SWidget> TraceControlPtr = TraceControl.Pin();
+				if (TraceControlPtr.IsValid())
+				{
+					FModuleManager::LoadModuleChecked<UE::TraceTools::ITraceToolsModule>("TraceTools").SetTraceControlWidgetInstanceId(TraceControlPtr.ToSharedRef(), InstanceId);
+				}
 			}
 		}
+
 		if (bIsSessionInfoSet)
 		{
 			UpdateAppTitle();
@@ -1026,6 +1041,58 @@ bool FInsightsManager::HandleResponseFileCmd(const TCHAR* ResponseFile, FOutputD
 
 ////////////////////////////////////////////////////////////////////////////////////////////////////
 
+void FInsightsManager::RegisterTraceControlTab()
+{
+#if !WITH_EDITOR
+	FTabSpawnerEntry& TabSpawnerEntry = FGlobalTabmanager::Get()->RegisterNomadTabSpawner(FInsightsManagerTabs::TraceControlTabId,
+		FOnSpawnTab::CreateRaw(this, &FInsightsManager::SpawnTraceControlTab))
+		.SetDisplayName(LOCTEXT("TraceControl", "Trace Control"))
+		.SetTooltipText(LOCTEXT("TraceControlTooltip", "Open the Trace Control tab."))
+		.SetIcon(FSlateIcon(FInsightsStyle::GetStyleSetName(), "Icons.TraceControl"))
+		.SetAutoGenerateMenuEntry(false);
+#endif
+}
+
+////////////////////////////////////////////////////////////////////////////////////////////////////
+
+TSharedRef<SDockTab> FInsightsManager::SpawnTraceControlTab(const FSpawnTabArgs& Args)
+{
+	const TSharedRef<SDockTab> DockTab = SNew(SDockTab)
+		.TabRole(ETabRole::NomadTab);
+
+	DockTab->SetOnTabClosed(SDockTab::FOnTabClosedCallback::CreateRaw(this, &FInsightsManager::OnTraceControlTabClosed));
+
+	ISessionServicesModule& SessionServicesModule = FModuleManager::LoadModuleChecked<ISessionServicesModule>("SessionServices");
+	TSharedPtr<ITraceController> TraceController = SessionServicesModule.GetTraceController();
+
+	TSharedRef<SWidget> TraceControlRef = FModuleManager::LoadModuleChecked<UE::TraceTools::ITraceToolsModule>("TraceTools").CreateTraceControlWidget(TraceController, InstanceId);
+
+	DockTab->SetContent(TraceControlRef);
+
+	TraceControl = TraceControlRef;
+
+	return DockTab;
+}
+
+////////////////////////////////////////////////////////////////////////////////////////////////////
+
+void FInsightsManager::OnTraceControlTabClosed(TSharedRef<SDockTab> TabBeingClosed)
+{
+	TraceControl.Reset();
+
+	// Disable TabClosed delegate.
+	TabBeingClosed->SetOnTabClosed(SDockTab::FOnTabClosedCallback());
+}
+
+////////////////////////////////////////////////////////////////////////////////////////////////////
+
+void FInsightsManager::OpenTraceControlWindow()
+{
+	FGlobalTabmanager::Get()->TryInvokeTab(FInsightsManagerTabs::TraceControlTabId);
+}
+
+////////////////////////////////////////////////////////////////////////////////////////////////////
+
 } // namespace UE::Insights
 
 #undef LOCTEXT_NAMESPACE
diff --git a/Engine/Source/Developer/TraceInsights/Private/Insights/InsightsManager.h b/Engine/Source/Developer/TraceInsights/Private/Insights/InsightsManager.h
--- a/Engine/Source/Developer/TraceInsights/Private/Insights/InsightsManager.h
+++ b/Engine/Source/Developer/TraceInsights/Private/Insights/InsightsManager.h
@@ -251,6 +251,8 @@ public:
 	/** Resets (closes) current session instance. */
 	void ResetSession(bool bNotify = true);
 
+	void OpenTraceControlWindow();
+
 	////////////////////////////////////////////////////////////////////////////////////////////////////
 	// SessionChangedEvent
 
@@ -298,6 +300,12 @@ private:
 
 	bool HandleResponseFileCmd(const TCHAR* ResponseFile, FOutputDevice& Ar);
 
+	void RegisterTraceControlTab();
+
+	TSharedRef<SDockTab> SpawnTraceControlTab(const FSpawnTabArgs& Args);
+
+	void OnTraceControlTabClosed(TSharedRef<SDockTab> TabBeingClosed);
+
 private:
 	bool bIsInitialized = false;
 
@@ -376,6 +384,9 @@ private:
 
 	/** The Trace Store connection */
 	UE::Trace::FStoreConnection TraceStoreConnection;
+
+	FGuid InstanceId;
+	TWeakPtr<SWidget> TraceControl;
 };
 
 } // namespace UE::Insights
diff --git a/Engine/Source/Developer/TraceInsights/Private/Insights/InsightsStyle.cpp b/Engine/Source/Developer/TraceInsights/Private/Insights/InsightsStyle.cpp
--- a/Engine/Source/Developer/TraceInsights/Private/Insights/InsightsStyle.cpp
+++ b/Engine/Source/Developer/TraceInsights/Private/Insights/InsightsStyle.cpp
@@ -273,6 +273,11 @@ void FInsightsStyle::FStyle::Initialize()
 	Set("Icons.NetEvent.TreeItem", new TODO_IMAGE_BRUSH(Icon12x12));
 	Set("Icons.MemTag.TreeItem", new TODO_IMAGE_BRUSH(Icon12x12));
 
+	//////////////////////////////////////////////////
+	// Trace Control
+
+	Set("Icons.TraceControl", new EDITOR_IMAGE_BRUSH_SVG("Starship/Common/TraceDataFiltering", Icon16x16));
+
 	//////////////////////////////////////////////////
 }
 
diff --git a/Engine/Source/Developer/TraceInsights/Public/Insights/IUnrealInsightsModule.h b/Engine/Source/Developer/TraceInsights/Public/Insights/IUnrealInsightsModule.h
--- a/Engine/Source/Developer/TraceInsights/Public/Insights/IUnrealInsightsModule.h
+++ b/Engine/Source/Developer/TraceInsights/Public/Insights/IUnrealInsightsModule.h
@@ -38,6 +38,7 @@ struct TRACEINSIGHTS_API FInsightsManagerTabs
 
 	static const FName AutomationWindowTabId;
 	static const FName MessageLogTabId;
+	static const FName TraceControlTabId;
 };
 
 ////////////////////////////////////////////////////////////////////////////////////////////////////
diff --git a/Engine/Source/Developer/TraceInsights/TraceInsights.Build.cs b/Engine/Source/Developer/TraceInsights/TraceInsights.Build.cs
--- a/Engine/Source/Developer/TraceInsights/TraceInsights.Build.cs
+++ b/Engine/Source/Developer/TraceInsights/TraceInsights.Build.cs
@@ -20,12 +20,14 @@ public class TraceInsights : ModuleRules
 				"InputCore",
 				"RenderCore",
 				"RHI",
+				"SessionServices",
 				"Slate",
 				"Sockets",
 				"SourceCodeAccess",
 				"TraceAnalysis",
 				"TraceInsightsCore",
 				"TraceLog",
+				"TraceTools",
 				"TraceServices",
 				"WorkspaceMenuStructure",
 				"XmlParser",
diff --git a/Engine/Source/Developer/TraceServices/Private/Analyzers/DiagnosticsAnalysis.cpp b/Engine/Source/Developer/TraceServices/Private/Analyzers/DiagnosticsAnalysis.cpp
--- a/Engine/Source/Developer/TraceServices/Private/Analyzers/DiagnosticsAnalysis.cpp
+++ b/Engine/Source/Developer/TraceServices/Private/Analyzers/DiagnosticsAnalysis.cpp
@@ -90,6 +90,13 @@ bool FDiagnosticsAnalyzer::OnEvent(uint16 RouteId, EStyle Style, const FOnEventC
 		SessionInfo.ConfigurationType = (EBuildConfiguration)EventData.GetValue<uint8>("ConfigurationType");
 		SessionInfo.TargetType = (EBuildTargetType)EventData.GetValue<uint8>("TargetType");
 
+		TArrayView<const uint32> Data = EventData.GetArrayView<uint32>("InstanceId");
+		if (Data.Num() > 0)
+		{
+			check(Data.Num() == 4);
+			SessionInfo.InstanceId = FGuid(Data[0], Data[1], Data[2], Data[3]);
+		}
+
 		FAnalysisSessionEditScope _(Session);
 		Provider->SetSessionInfo(SessionInfo);
 		UpdateSessionMetadata(EventData);
diff --git a/Engine/Source/Developer/TraceServices/Public/TraceServices/Model/Diagnostics.h b/Engine/Source/Developer/TraceServices/Public/TraceServices/Model/Diagnostics.h
--- a/Engine/Source/Developer/TraceServices/Public/TraceServices/Model/Diagnostics.h
+++ b/Engine/Source/Developer/TraceServices/Public/TraceServices/Model/Diagnostics.h
@@ -22,6 +22,7 @@ struct FSessionInfo
 	uint32 Changelist = 0;
 	EBuildConfiguration ConfigurationType = EBuildConfiguration::Unknown;
 	EBuildTargetType TargetType = EBuildTargetType::Unknown;
+	FGuid InstanceId;
 };
 
 class IDiagnosticsProvider : public IProvider
diff --git a/Engine/Source/Developer/TraceTools/Private/Services/ISessionTraceFilterService.h b/Engine/Source/Developer/TraceTools/Private/Services/ISessionTraceFilterService.h
--- a/Engine/Source/Developer/TraceTools/Private/Services/ISessionTraceFilterService.h
+++ b/Engine/Source/Developer/TraceTools/Private/Services/ISessionTraceFilterService.h
@@ -75,6 +75,12 @@ public:
 
 	/** Get the current status of the trace system. */
 	virtual FTraceStatus::ETraceSystemStatus GetTraceSystemStatus() const = 0;
+
+	/** Sets the InstanceId to control. Supports invalid guid value for disabled state. */
+	virtual void SetInstanceId(const FGuid& Id) = 0;
+
+	/** Returns true if the session with the currently set InstanceId is available for communication. */
+	virtual bool HasAvailableInstance() const = 0;
 };
 
 } // namespace UE::TraceTools
\ No newline at end of file
diff --git a/Engine/Source/Developer/TraceTools/Private/Services/SessionTraceControllerFilterService.cpp b/Engine/Source/Developer/TraceTools/Private/Services/SessionTraceControllerFilterService.cpp
--- a/Engine/Source/Developer/TraceTools/Private/Services/SessionTraceControllerFilterService.cpp
+++ b/Engine/Source/Developer/TraceTools/Private/Services/SessionTraceControllerFilterService.cpp
@@ -26,16 +26,13 @@ FSessionTraceControllerFilterService::FSessionTraceControllerFilterService(TShar
 	FCoreDelegates::OnEndFrame.AddRaw(this, &FSessionTraceControllerFilterService::OnApplyChannelChanges);
 
 	TraceController = InTraceController;
-	TraceController->OnSelectedSessionStatusReceived().AddRaw(this, &FSessionTraceControllerFilterService::OnTraceStatusUpdated);
-	TraceController->OnSessionSelectionChanged().AddRaw(this, &FSessionTraceControllerFilterService::OnSessionSelectionChanged);
+	TraceController->OnStatusReceived().AddRaw(this, &FSessionTraceControllerFilterService::OnTraceStatusUpdated);
 }
 
 FSessionTraceControllerFilterService::~FSessionTraceControllerFilterService()
 {
 	FCoreDelegates::OnEndFrame.RemoveAll(this);
-
-	TraceController->OnSelectedSessionStatusReceived().RemoveAll(this);
-	TraceController->OnSessionSelectionChanged().RemoveAll(this);
+	TraceController->OnStatusReceived().RemoveAll(this);
 }
 
 void FSessionTraceControllerFilterService::GetRootObjects(TArray<FTraceObjectInfo>& OutObjects) const
@@ -104,11 +101,16 @@ void FSessionTraceControllerFilterService::DisableAllChannels()
 
 void FSessionTraceControllerFilterService::OnTraceStatusUpdated(const FTraceStatus& InStatus, FTraceStatus::EUpdateType InUpdateType, ITraceControllerCommands& Commands)
 {
-	if (!TraceController->HasAvailableSelectedInstance())
+	if (!InstanceId.IsValid() || InStatus.InstanceId != InstanceId)
+	{
+		return;
+	}
+
+	if (!TraceController->HasAvailableInstance(InstanceId))
 	{
-		Objects.Empty();
 		return;
 	}
+
 	if (EnumHasAnyFlags(InUpdateType, FTraceStatus::EUpdateType::ChannelsDesc) ||
 		(EnumHasAnyFlags(InUpdateType, FTraceStatus::EUpdateType::ChannelsStatus)))
 	{
@@ -171,14 +173,14 @@ void FSessionTraceControllerFilterService::UpdateChannels(const FTraceStatus& In
 
 void FSessionTraceControllerFilterService::OnApplyChannelChanges()
 {
-	if (!TraceController->HasAvailableSelectedInstance() || !bChannelsReceived)
+	if (!TraceController->HasAvailableInstance(InstanceId) || !bChannelsReceived)
 	{
 		return;
 	}
 
 	if (FrameEnabledChannels.Num() || FrameDisabledChannels.Num())
 	{
-		TraceController->WithSelectedInstances([&](const FTraceStatus& Status, ITraceControllerCommands& Commands)
+		TraceController->WithInstance(InstanceId, [&](const FTraceStatus& Status, ITraceControllerCommands& Commands)
 		{
 			Commands.SetChannels(FrameEnabledChannels.Array(), FrameDisabledChannels.Array());
 		});
@@ -208,11 +210,27 @@ const FTraceStats& FSessionTraceControllerFilterService::GetStats() const
 	return Stats;
 }
 
-void FSessionTraceControllerFilterService::OnSessionSelectionChanged()
+void FSessionTraceControllerFilterService::Reset()
 {
 	bHasStats = false;
 	bHasSettings = false;
 	TraceEndpoint.Empty();
 }
 
+void FSessionTraceControllerFilterService::SetInstanceId(const FGuid& Id)
+{
+	InstanceId = Id;
+	Reset();
+}
+
+bool FSessionTraceControllerFilterService::HasAvailableInstance() const
+{
+	if (InstanceId.IsValid())
+	{
+		return TraceController->HasAvailableInstance(InstanceId);
+	}
+
+	return false;
+}
+
 } // namespace UE::TraceTools
\ No newline at end of file
diff --git a/Engine/Source/Developer/TraceTools/Private/Services/SessionTraceControllerFilterService.h b/Engine/Source/Developer/TraceTools/Private/Services/SessionTraceControllerFilterService.h
--- a/Engine/Source/Developer/TraceTools/Private/Services/SessionTraceControllerFilterService.h
+++ b/Engine/Source/Developer/TraceTools/Private/Services/SessionTraceControllerFilterService.h
@@ -39,6 +39,10 @@ public:
 
 	virtual const FString& GetTraceEndpoint() const override { return TraceEndpoint; }
 	virtual FTraceStatus::ETraceSystemStatus GetTraceSystemStatus() const { return TraceSystemStatus; }
+
+	void SetInstanceId(const FGuid& Id) override;
+	bool HasAvailableInstance() const override;
+
 	/** End ISessionTraceFilterService overrides */
 
 protected:
@@ -52,7 +56,7 @@ protected:
 
 	void UpdateChannels(const FTraceStatus& InStatus);
 
-	void OnSessionSelectionChanged();
+	void Reset();
 
 protected:
 	TSharedPtr<ITraceController> TraceController;
@@ -80,6 +84,8 @@ protected:
 
 	FString TraceEndpoint;
 	FTraceStatus::ETraceSystemStatus TraceSystemStatus;
+
+	FGuid InstanceId;
 };
 
 } // namespace UE::TraceTools
\ No newline at end of file
diff --git a/Engine/Source/Developer/TraceTools/Private/TraceToolsModule.cpp b/Engine/Source/Developer/TraceTools/Private/TraceToolsModule.cpp
--- a/Engine/Source/Developer/TraceTools/Private/TraceToolsModule.cpp
+++ b/Engine/Source/Developer/TraceTools/Private/TraceToolsModule.cpp
@@ -32,5 +32,20 @@ void FTraceToolsModule::ShutdownModule()
 
 TSharedRef<SWidget> FTraceToolsModule::CreateTraceControlWidget(TSharedPtr<ITraceController> InTraceController)
 {
-	return SNew(UE::TraceTools::STraceControl, InTraceController);
+	return SNew(UE::TraceTools::STraceControl, InTraceController)
+		   .AutoDetectSelectedSession(true);
+}
+
+TSharedRef<SWidget> FTraceToolsModule::CreateTraceControlWidget(TSharedPtr<ITraceController> InTraceController, FGuid InstanceId)
+{
+	TSharedRef<UE::TraceTools::STraceControl> TraceControlWidget = SNew(UE::TraceTools::STraceControl, InTraceController);
+	TraceControlWidget->SetInstanceId(InstanceId);
+
+	return TraceControlWidget;
+}
+
+void FTraceToolsModule::SetTraceControlWidgetInstanceId(TSharedRef<SWidget> Widget, FGuid InstanceId)
+{
+	TSharedRef<UE::TraceTools::STraceControl> TraceControl = StaticCastSharedRef<UE::TraceTools::STraceControl>(Widget);
+	TraceControl->SetInstanceId(InstanceId);
 }
\ No newline at end of file
diff --git a/Engine/Source/Developer/TraceTools/Private/TraceToolsModule.h b/Engine/Source/Developer/TraceTools/Private/TraceToolsModule.h
--- a/Engine/Source/Developer/TraceTools/Private/TraceToolsModule.h
+++ b/Engine/Source/Developer/TraceTools/Private/TraceToolsModule.h
@@ -17,7 +17,10 @@ public:
 	virtual void StartupModule() override;
 	virtual void ShutdownModule() override;
 
-	virtual TSharedRef<SWidget> CreateTraceControlWidget(TSharedPtr<ITraceController> InTraceController);
+	virtual TSharedRef<SWidget> CreateTraceControlWidget(TSharedPtr<ITraceController> InTraceController) override;
+	virtual TSharedRef<SWidget> CreateTraceControlWidget(TSharedPtr<ITraceController> InTraceController, FGuid InstanceId) override;
+
+	virtual void SetTraceControlWidgetInstanceId(TSharedRef<SWidget> Widget, FGuid InstanceId) override;
 
 	static FString TraceFiltersIni;
 };
\ No newline at end of file
diff --git a/Engine/Source/Developer/TraceTools/Private/Widgets/STraceControl.cpp b/Engine/Source/Developer/TraceTools/Private/Widgets/STraceControl.cpp
--- a/Engine/Source/Developer/TraceTools/Private/Widgets/STraceControl.cpp
+++ b/Engine/Source/Developer/TraceTools/Private/Widgets/STraceControl.cpp
@@ -3,10 +3,19 @@
 #include "STraceControl.h"
 
 #include "Framework/Commands/UICommandList.h"
+#include "ISessionInstanceInfo.h"
+#include "ISessionManager.h"
+#include "ISessionServicesModule.h"
+#include "Modules/ModuleManager.h"
 #include "STraceDataFilterWidget.h"
 #include "Widgets/SBoxPanel.h"
 #include "Widgets/STraceControlToolbar.h"
 
+// Trace Tools
+#include "Services/SessionTraceControllerFilterService.h"
+
+#define LOCTEXT_NAMESPACE "UE::TraceTools::STraceControl"
+
 namespace UE::TraceTools
 {
 
@@ -16,19 +25,24 @@ STraceControl::STraceControl()
 
 STraceControl::~STraceControl()
 {
-
+	if (bAutoDetectSelectedSession && SessionManager.IsValid())
+	{
+		SessionManager->OnInstanceSelectionChanged().RemoveAll(this);
+	}
 }
 
 void STraceControl::Construct(const FArguments& InArgs, TSharedPtr<ITraceController> InTraceController)
 {
 	TraceController = InTraceController;
+	bAutoDetectSelectedSession = InArgs._AutoDetectSelectedSession;
 
 	UICommandList = MakeShareable(new FUICommandList);
-	BindCommands();
 
 	TraceController->SendStatusUpdateRequest();
 	TraceController->SendChannelUpdateRequest();
 
+	SessionFilterService = MakeShareable(new FSessionTraceControllerFilterService(InTraceController));
+
 	ChildSlot
 	[
 		SNew(SVerticalBox)
@@ -37,20 +51,76 @@ void STraceControl::Construct(const FArguments& InArgs, TSharedPtr<ITraceControl
 		.Padding(0.0f, 4.0f, 0.0f, 0.0f)
 		.AutoHeight()
 		[
-			SNew(STraceControlToolbar, UICommandList.ToSharedRef(), InTraceController)
+			SAssignNew(Toolbar, STraceControlToolbar, UICommandList.ToSharedRef(), InTraceController)
 		]
 
 		+ SVerticalBox::Slot()
 		.Padding(FMargin(0.0f, 0.0f, 0.0f, 2.0f))
 		[
-			SNew(STraceDataFilterWidget, InTraceController)
+			SAssignNew(TraceDataFilterWidget, STraceDataFilterWidget, InTraceController, SessionFilterService)
 		]
 	];
+
+	if (bAutoDetectSelectedSession)
+	{
+		ISessionServicesModule& SessionServicesModule = FModuleManager::LoadModuleChecked<ISessionServicesModule>("SessionServices");
+		SessionManager = SessionServicesModule.GetSessionManager();
+
+		if (SessionManager.IsValid())
+		{
+			//Note that we should not add ourselves as shared pointer to session manager,
+			//that will create a circular dependency
+			SessionManager->OnInstanceSelectionChanged().AddRaw(this, &STraceControl::OnInstanceSelectionChanged);
+		}
+
+		TraceDataFilterWidget->SetWarningBannerText(LOCTEXT("NoSessionSelectedWarning", "Please select an active instance from the Session Browser."));
+
+		const TArray<TSharedPtr<ISessionInstanceInfo>> CurrentSelectedSessions = SessionManager->GetSelectedInstances();
+		for (TSharedPtr<ISessionInstanceInfo> Instance : CurrentSelectedSessions)
+		{
+			OnInstanceSelectionChanged(Instance, true);
+		}
+	}
 }
 
-void STraceControl::BindCommands()
+void STraceControl::SetInstanceId(const FGuid& Id)
 {
+	Toolbar->SetInstanceId(Id);
+	SessionFilterService->SetInstanceId(Id);
 
+	if (Id.IsValid())
+	{
+		TraceController->SendDiscoveryRequest(FGuid(), Id);
+	}
+}
+
+void STraceControl::OnInstanceSelectionChanged(const TSharedPtr<ISessionInstanceInfo>& InstanceInfo, bool bSelected)
+{
+	if (!InstanceInfo.IsValid())
+	{
+		return;
+	}
+
+	if (bSelected)
+	{
+		SelectedSessionsIds.Add(InstanceInfo->GetInstanceId());
+	}
+	else
+	{
+		SelectedSessionsIds.Remove(InstanceInfo->GetInstanceId());
+	}
+
+	if (SelectedSessionsIds.Num() == 1)
+	{
+		SetInstanceId(*SelectedSessionsIds.begin());
+	}
+	else
+	{
+		InstanceId.Invalidate();
+		SetInstanceId(InstanceId);
+	}
 }
 
 } // namespace UE::TraceTools
+
+#undef LOCTEXT_NAMESPACE // UE::TraceTools::STraceControl
\ No newline at end of file
diff --git a/Engine/Source/Developer/TraceTools/Private/Widgets/STraceControl.h b/Engine/Source/Developer/TraceTools/Private/Widgets/STraceControl.h
--- a/Engine/Source/Developer/TraceTools/Private/Widgets/STraceControl.h
+++ b/Engine/Source/Developer/TraceTools/Private/Widgets/STraceControl.h
@@ -5,12 +5,17 @@
 #include "ITraceController.h"
 #include "Widgets/SCompoundWidget.h"
 
+class ISessionInstanceInfo;
 class ISessionManager;
 class FUICommandList;
 
 namespace UE::TraceTools
 {
 
+class ISessionTraceFilterService;
+class STraceControlToolbar;
+class STraceDataFilterWidget;
+
 class STraceControl : public SCompoundWidget
 {
 public:
@@ -20,19 +25,44 @@ public:
 	/** Virtual destructor. */
 	virtual ~STraceControl();
 
-	SLATE_BEGIN_ARGS(STraceControl) {}
+	SLATE_BEGIN_ARGS(STraceControl) 
+		: _AutoDetectSelectedSession(false)
+	{}
+
+	/* Specifies if the widget should autodetect the selected session in Session Browser. If false, the SessionId will need to be set manually. */
+	SLATE_ARGUMENT(bool, AutoDetectSelectedSession)
+
 	SLATE_END_ARGS()
 
 	/** Constructs this widget. */
 	void Construct(const FArguments& InArgs, TSharedPtr<ITraceController> InTraceController);
 
+	/** Sets the InstanceId to control. Supports invalid guid value for disabled state. */
+	void SetInstanceId(const FGuid& Id);
+
 protected:
-	void BindCommands();
+	void OnInstanceSelectionChanged(const TSharedPtr<class ISessionInstanceInfo>&, bool);
 
 private:
 	TSharedPtr<ITraceController> TraceController;
 
 	TSharedPtr<FUICommandList> UICommandList;
+
+	/** The InstanceId to control. */
+	FGuid InstanceId;
+
+	TSharedPtr<STraceControlToolbar> Toolbar;
+
+	TSharedPtr<STraceDataFilterWidget> TraceDataFilterWidget;
+
+	/** Session manager used for selecting sessions */
+	TSharedPtr<ISessionManager> SessionManager;
+
+	TSharedPtr<ISessionTraceFilterService> SessionFilterService;
+
+	TSet<FGuid> SelectedSessionsIds;
+
+	bool bAutoDetectSelectedSession;
 };
 
 } // namespace UE::TraceTools
diff --git a/Engine/Source/Developer/TraceTools/Private/Widgets/STraceControlToolbar.cpp b/Engine/Source/Developer/TraceTools/Private/Widgets/STraceControlToolbar.cpp
--- a/Engine/Source/Developer/TraceTools/Private/Widgets/STraceControlToolbar.cpp
+++ b/Engine/Source/Developer/TraceTools/Private/Widgets/STraceControlToolbar.cpp
@@ -29,7 +29,7 @@ STraceControlToolbar::STraceControlToolbar()
 
 STraceControlToolbar::~STraceControlToolbar()
 {
-	TraceController->OnSelectedSessionStatusReceived().Remove(OnStatusReceivedDelegate);
+	TraceController->OnStatusReceived().RemoveAll(this);
 }
 
 BEGIN_SLATE_FUNCTION_BUILD_OPTIMIZATION
@@ -40,7 +40,7 @@ void STraceControlToolbar::Construct(const FArguments& InArgs, const TSharedRef<
 	FTraceControlCommands::Register();
 	TraceController = InTraceController;
 
-	OnStatusReceivedDelegate = TraceController->OnSelectedSessionStatusReceived().AddSP(this, &STraceControlToolbar::OnTraceStatusUpdated);
+	OnStatusReceivedDelegate = TraceController->OnStatusReceived().AddRaw(this, &STraceControlToolbar::OnTraceStatusUpdated);
 
 	BindCommands(CommandList);
 
@@ -186,7 +186,7 @@ void STraceControlToolbar::InitializeSettings()
 
 bool STraceControlToolbar::SetTraceTarget_CanExecute() const
 {
-	return TraceController->HasAvailableSelectedInstance() && !bIsTracing;
+	return IsInstanceAvailable() && !bIsTracing;
 }
 
 void STraceControlToolbar::SetTraceTarget_Execute(ETraceTarget InTraceTarget)
@@ -194,16 +194,21 @@ void STraceControlToolbar::SetTraceTarget_Execute(ETraceTarget InTraceTarget)
 	TraceTarget = InTraceTarget;
 }
 
+bool STraceControlToolbar::IsInstanceAvailable() const
+{
+	return InstanceId.IsValid() && TraceController->HasAvailableInstance(InstanceId) && bIsTracingAvailable;
+}
+
 bool STraceControlToolbar::ToggleTrace_CanExecute() const
 {
-	return TraceController->HasAvailableSelectedInstance() && bIsTracingAvailable;
+	return IsInstanceAvailable();
 }
 
 void STraceControlToolbar::ToggleTrace_Execute()
 {
 	if (!bIsTracing)
 	{
-		TraceController->WithSelectedInstances([&](const FTraceStatus& Status, ITraceControllerCommands& Commands)
+		TraceController->WithInstance(InstanceId, [&](const FTraceStatus& Status, ITraceControllerCommands& Commands)
 			{
 				if (TraceTarget == ETraceTarget::Server)
 				{
@@ -218,7 +223,7 @@ void STraceControlToolbar::ToggleTrace_Execute()
 	}
 	else
 	{
-		TraceController->WithSelectedInstances([&](const FTraceStatus& Status, ITraceControllerCommands& Commands)
+		TraceController->WithInstance(InstanceId, [&](const FTraceStatus& Status, ITraceControllerCommands& Commands)
 			{
 				Commands.Stop();
 			});
@@ -228,12 +233,12 @@ void STraceControlToolbar::ToggleTrace_Execute()
 
 bool STraceControlToolbar::TraceSnapshot_CanExecute() const
 {
-	return TraceController->HasAvailableSelectedInstance() && bIsTracingAvailable;
+	return IsInstanceAvailable();
 }
 
 void STraceControlToolbar::TraceSnapshot_Execute()
 {
-	TraceController->WithSelectedInstances([&](const FTraceStatus& Status, ITraceControllerCommands& Commands)
+	TraceController->WithInstance(InstanceId, [&](const FTraceStatus& Status, ITraceControllerCommands& Commands)
 	{
 		if (TraceTarget == ETraceTarget::Server)
 		{
@@ -248,13 +253,13 @@ void STraceControlToolbar::TraceSnapshot_Execute()
 
 bool STraceControlToolbar::TraceBookmark_CanExecute() const
 {
-	return TraceController->HasAvailableSelectedInstance() && bIsTracingAvailable && bIsTracing && !bIsPaused;
+	return IsInstanceAvailable() && bIsTracing && !bIsPaused;
 }
 
 void STraceControlToolbar::TraceBookmark_Execute()
 {
 	const FString BookmarkName = FDateTime::Now().ToString(TEXT("Bookmark_%Y%m%d_%H%M%S"));
-	TraceController->WithSelectedInstances([&](const FTraceStatus& Status, ITraceControllerCommands& Commands)
+	TraceController->WithInstance(InstanceId, [&](const FTraceStatus& Status, ITraceControllerCommands& Commands)
 	{
 		Commands.Bookmark(BookmarkName);
 	});
@@ -262,12 +267,12 @@ void STraceControlToolbar::TraceBookmark_Execute()
 
 bool STraceControlToolbar::TraceScreenshot_CanExecute() const
 {
-	return TraceController->HasAvailableSelectedInstance() && bIsTracing && !bIsPaused;
+	return IsInstanceAvailable() && bIsTracing && !bIsPaused;
 }
 
 void STraceControlToolbar::TraceScreenshot_Execute()
 {
-	TraceController->WithSelectedInstances([&](const FTraceStatus& Status, ITraceControllerCommands& Commands)
+	TraceController->WithInstance(InstanceId, [&](const FTraceStatus& Status, ITraceControllerCommands& Commands)
 	{
 		Commands.Screenshot(TEXT(""), false);
 	});
@@ -275,7 +280,7 @@ void STraceControlToolbar::TraceScreenshot_Execute()
 
 bool STraceControlToolbar::ToggleStatNamedEvents_CanExecute() const
 {
-	return TraceController->HasAvailableSelectedInstance() && bIsTracing;
+	return IsInstanceAvailable();
 }
 
 bool STraceControlToolbar::ToggleStatNamedEvents_IsChecked() const
@@ -286,7 +291,7 @@ bool STraceControlToolbar::ToggleStatNamedEvents_IsChecked() const
 void STraceControlToolbar::ToggleStatNamedEvents_Execute()
 {
 	bAreStatNamedEventsEnabled = !bAreStatNamedEventsEnabled;
-	TraceController->WithSelectedInstances([&](const FTraceStatus& Status, ITraceControllerCommands& Commands)
+	TraceController->WithInstance(InstanceId, [&](const FTraceStatus& Status, ITraceControllerCommands& Commands)
 	{
 		Commands.SetStatNamedEventsEnabled(bAreStatNamedEventsEnabled);
 	});
@@ -319,6 +324,11 @@ FSlateIcon STraceControlToolbar::GetTraceTargetIcon() const
 
 void STraceControlToolbar::OnTraceStatusUpdated(const FTraceStatus& InStatus, FTraceStatus::EUpdateType InUpdateType, ITraceControllerCommands& Commands)
 {
+	if (!InstanceId.IsValid() || InstanceId != InStatus.InstanceId)
+	{
+		return;
+	}
+
 	bIsTracing = InStatus.bIsTracing;
 	bIsPaused = InStatus.bIsPaused;
 	bAreStatNamedEventsEnabled = InStatus.bAreStatNamedEventsEnabled;
@@ -336,14 +346,14 @@ void STraceControlToolbar::OnTraceStatusUpdated(const FTraceStatus& InStatus, FT
 
 bool STraceControlToolbar::TogglePauseResume_CanExecute() const
 {
-	return TraceController->HasAvailableSelectedInstance() && bIsTracing;
+	return IsInstanceAvailable() && bIsTracing;
 }
 
 FReply STraceControlToolbar::TogglePauseResume_OnClicked()
 {
 	if (bIsPaused)
 	{
-		TraceController->WithSelectedInstances([&](const FTraceStatus& Status, ITraceControllerCommands& Commands)
+		TraceController->WithInstance(InstanceId, [&](const FTraceStatus& Status, ITraceControllerCommands& Commands)
 			{
 				Commands.Resume();
 			});
@@ -351,7 +361,7 @@ FReply STraceControlToolbar::TogglePauseResume_OnClicked()
 	}
 	else
 	{
-		TraceController->WithSelectedInstances([&](const FTraceStatus& Status, ITraceControllerCommands& Commands)
+		TraceController->WithInstance(InstanceId, [&](const FTraceStatus& Status, ITraceControllerCommands& Commands)
 			{
 				Commands.Pause();
 			});
diff --git a/Engine/Source/Developer/TraceTools/Private/Widgets/STraceControlToolbar.h b/Engine/Source/Developer/TraceTools/Private/Widgets/STraceControlToolbar.h
--- a/Engine/Source/Developer/TraceTools/Private/Widgets/STraceControlToolbar.h
+++ b/Engine/Source/Developer/TraceTools/Private/Widgets/STraceControlToolbar.h
@@ -40,6 +40,8 @@ public:
 
 	void Construct( const FArguments& InArgs, const TSharedRef<FUICommandList>& CommandList, TSharedPtr<ITraceController> InTraceController);
 
+	void SetInstanceId(const FGuid& Id) { InstanceId = Id; }
+
 private:
 	void BindCommands(const TSharedRef<FUICommandList>& CommandList);
 
@@ -76,6 +78,8 @@ private:
 
 	const FSlateBrush* GetPauseResumeBrush() const;
 
+	bool IsInstanceAvailable() const;
+
 private:
 	TSharedPtr<ITraceController> TraceController;
 
@@ -85,6 +89,7 @@ private:
 	bool bAreStatNamedEventsEnabled = false;
 	bool bIsTracingAvailable = false;
 	FString TraceHostAddr;
+	FGuid InstanceId;
 
 	FDelegateHandle OnStatusReceivedDelegate;
 };
diff --git a/Engine/Source/Developer/TraceTools/Private/Widgets/STraceDataFilterWidget.cpp b/Engine/Source/Developer/TraceTools/Private/Widgets/STraceDataFilterWidget.cpp
--- a/Engine/Source/Developer/TraceTools/Private/Widgets/STraceDataFilterWidget.cpp
+++ b/Engine/Source/Developer/TraceTools/Private/Widgets/STraceDataFilterWidget.cpp
@@ -34,15 +34,12 @@ STraceDataFilterWidget::STraceDataFilterWidget() : bNeedsListRefresh(false), bHi
 
 STraceDataFilterWidget::~STraceDataFilterWidget()
 {
-	TraceController->OnSessionSelectionChanged().RemoveAll(this);
 }
 
-void STraceDataFilterWidget::Construct(const FArguments& InArgs, TSharedPtr<ITraceController> InTraceController)
+void STraceDataFilterWidget::Construct(const FArguments& InArgs, TSharedPtr<ITraceController> InTraceController, TSharedPtr<ISessionTraceFilterService> InSessionFilterService)
 {
 	TraceController = InTraceController;
-	SessionFilterService = MakeShareable(new FSessionTraceControllerFilterService(InTraceController));
-
-	TraceController->OnSessionSelectionChanged().AddSP(this, &STraceDataFilterWidget::OnSessionSelectionChanged);
+	SessionFilterService = InSessionFilterService;
 
 	SAssignNew(ExternalScrollbar, SScrollBar)
 	.AlwaysShowScrollbar(true);
@@ -151,7 +148,7 @@ void STraceDataFilterWidget::Construct(const FArguments& InArgs, TSharedPtr<ITra
 					.Padding(4.0f, 0.0f, 0.0f, 0.0f)
 					[
 						SNew(STextBlock)
-						.Text(LOCTEXT("NoSessionSelectedWarning", "Please select an active instance from the Session Browser."))
+						.Text(this, &STraceDataFilterWidget::GetWarningBannerText)
 						.ColorAndOpacity(FLinearColor(1.0f, 1.0f, 1.0f, 1.0f))
 					]
 				]
@@ -436,12 +433,12 @@ TSharedRef<ITraceObject> STraceDataFilterWidget::AddFilterableObject(const FTrac
 
 bool STraceDataFilterWidget::HasValidData() const
 {
-	return bHasChannelData && TraceController->GetNumSelectedInstances() == 1 && TraceController->HasAvailableSelectedInstance() && SessionFilterService.IsValid();
+	return bHasChannelData && SessionFilterService.IsValid() && SessionFilterService->HasAvailableInstance();
 }
 
 bool STraceDataFilterWidget::ShouldShowBanner() const
 {
-	return TraceController->GetNumSelectedInstances() != 1 || !TraceController->HasAvailableSelectedInstance() || !SessionFilterService.IsValid();
+	return !HasValidData();
 }
 
 void STraceDataFilterWidget::EnumerateSelectedItems(TFunction<void(TSharedPtr<ITraceObject> InItem)> InFunction) const
@@ -574,6 +571,16 @@ void STraceDataFilterWidget::OnSessionSelectionChanged()
 	bHasSettings = false;
 }
 
+FText STraceDataFilterWidget::GetWarningBannerText() const
+{
+	if (!WarningBannerText.IsEmpty())
+	{
+		return WarningBannerText;
+	}
+
+	return LOCTEXT("ConnectingToSessionWarning", "Connecting to live session.");
+}
+
 } // namespace UE::TraceTools
 
 #undef LOCTEXT_NAMESPACE // "STraceDataFilterWidget"
diff --git a/Engine/Source/Developer/TraceTools/Private/Widgets/STraceDataFilterWidget.h b/Engine/Source/Developer/TraceTools/Private/Widgets/STraceDataFilterWidget.h
--- a/Engine/Source/Developer/TraceTools/Private/Widgets/STraceDataFilterWidget.h
+++ b/Engine/Source/Developer/TraceTools/Private/Widgets/STraceDataFilterWidget.h
@@ -39,10 +39,12 @@ public:
 	SLATE_END_ARGS()
 
 	/** Constructs this widget. */
-	void Construct(const FArguments& InArgs, TSharedPtr<ITraceController> InTraceController);
+	void Construct(const FArguments& InArgs, TSharedPtr<ITraceController> InTraceController, TSharedPtr<ISessionTraceFilterService> InSessionFilterService);
 
 	virtual void Tick(const FGeometry& AllottedGeometry, const double InCurrentTime, const float InDeltaTime) override;
 
+	void SetWarningBannerText(const FText& InText) { WarningBannerText = InText; };
+
 protected:
 	/** Callback from SFilterPresetList, should save the current tileview filter state as the specified preset */
 	void OnSavePreset(const TSharedPtr<ITraceFilterPreset>& Preset);
@@ -85,6 +87,8 @@ protected:
 
 	void OnSessionSelectionChanged();
 
+	FText GetWarningBannerText() const;
+
 protected:
 	/** Flat list of contained items */
 	TArray<TSharedPtr<ITraceObject>> ListItems;
@@ -127,6 +131,8 @@ protected:
 
 	bool bHasChannelData = false;
 	bool bHasSettings = false;
+
+	FText WarningBannerText;
 };
 
 } // namespace UE::TraceTools
\ No newline at end of file
diff --git a/Engine/Source/Developer/TraceTools/Public/TraceTools/Interfaces/ITraceToolsModule.h b/Engine/Source/Developer/TraceTools/Public/TraceTools/Interfaces/ITraceToolsModule.h
--- a/Engine/Source/Developer/TraceTools/Public/TraceTools/Interfaces/ITraceToolsModule.h
+++ b/Engine/Source/Developer/TraceTools/Public/TraceTools/Interfaces/ITraceToolsModule.h
@@ -22,16 +22,26 @@ class ITraceToolsModule
 	: public IModuleInterface
 {
 public:
-	
 	/**
-	 * Creates a trace control widget.
+	 * Creates a trace control widget that auto detects the session selected in Session Browser and can be used to control it.
 	 *
 	 * @return The new widget
 	 */
 	virtual TSharedRef<SWidget> CreateTraceControlWidget(TSharedPtr<ITraceController> InTraceController) = 0;
 
-public:
+	/**
+	 * Creates a trace control widget that controls a specific Instance Id.
+	 *
+	 * @return The new widget
+	 */
+	virtual TSharedRef<SWidget> CreateTraceControlWidget(TSharedPtr<ITraceController> InTraceController, FGuid InstanceId) = 0;
 
+	/**
+	 * Sets the InstanceId to control for the provided widget, which must be a widget returned by CreateTraceControlWidget.
+	 */
+	virtual void SetTraceControlWidgetInstanceId(TSharedRef<SWidget> Widget, FGuid InstanceId) = 0;
+
+public:
 	/**
 	 * Virtual destructor.
 	 */
diff --git a/Engine/Source/Runtime/Core/Private/ProfilingDebugging/TraceAuxiliary.cpp b/Engine/Source/Runtime/Core/Private/ProfilingDebugging/TraceAuxiliary.cpp
--- a/Engine/Source/Runtime/Core/Private/ProfilingDebugging/TraceAuxiliary.cpp
+++ b/Engine/Source/Runtime/Core/Private/ProfilingDebugging/TraceAuxiliary.cpp
@@ -1316,6 +1316,7 @@ UE_TRACE_EVENT_BEGIN(Diagnostics, Session2, NoSync|Important)
 	UE_TRACE_EVENT_FIELD(uint32, Changelist)
 	UE_TRACE_EVENT_FIELD(uint8, ConfigurationType)
 	UE_TRACE_EVENT_FIELD(uint8, TargetType)
+	UE_TRACE_EVENT_FIELD(uint32[], InstanceId)
 UE_TRACE_EVENT_END()
 
 ////////////////////////////////////////////////////////////////////////////////
@@ -1600,6 +1601,8 @@ void FTraceAuxiliary::Initialize(const TCHAR* CommandLine)
 	}
 #endif
 
+	constexpr int InstanceIdSize = 4;
+
 	// Trace out information about this session. This is done before initialization,
 	// so that it is always sent (all channels are enabled prior to initialization).
 	const TCHAR* BranchName = BuildSettings::GetBranchName();
@@ -1616,7 +1619,16 @@ void FTraceAuxiliary::Initialize(const TCHAR* CommandLine)
 		(ProjectNameLen * sizeof(TCHAR)) +
 		(CommandLineLen * sizeof(TCHAR)) +
 		(BranchNameLen * sizeof(TCHAR)) +
-		(BuildVersionLen * sizeof(TCHAR));
+		(BuildVersionLen * sizeof(TCHAR)) +
+		(InstanceIdSize * sizeof(uint32));
+
+	FGuid InstanceGuid = FApp::GetInstanceId();
+	uint32 InstanceId[InstanceIdSize];
+	for (int Index = 0; Index < InstanceIdSize; ++Index)
+	{
+		InstanceId[Index] = InstanceGuid[Index];
+	}
+	
 	UE_TRACE_LOG(Diagnostics, Session2, UE::Trace::TraceLogChannel, DataSize)
 		<< Session2.Platform(PREPROCESSOR_TO_STRING(UBT_COMPILED_PLATFORM), PlatformLen)
 		<< Session2.AppName(AppName, AppNameLen)
@@ -1626,7 +1638,8 @@ void FTraceAuxiliary::Initialize(const TCHAR* CommandLine)
 		<< Session2.BuildVersion(BuildVersion, BuildVersionLen)
 		<< Session2.Changelist(BuildSettings::GetCurrentChangelist())
 		<< Session2.ConfigurationType(uint8(FApp::GetBuildConfiguration()))
-		<< Session2.TargetType(uint8(FApp::GetBuildTargetType()));
+		<< Session2.TargetType(uint8(FApp::GetBuildTargetType()))
+		<< Session2.InstanceId(InstanceId, InstanceIdSize);
 
 	// Attempt to send trace data somewhere from the command line. It perhaps
 	// seems odd to do this before initializing Trace, but it is done this way
diff --git a/Engine/Source/Runtime/Engine/Private/TraceService.cpp b/Engine/Source/Runtime/Engine/Private/TraceService.cpp
--- a/Engine/Source/Runtime/Engine/Private/TraceService.cpp
+++ b/Engine/Source/Runtime/Engine/Private/TraceService.cpp
@@ -289,7 +289,7 @@ void FTraceServiceImpl::OnSettingsPing(const FTraceControlSettingsPing& Message,
 
 void FTraceServiceImpl::OnDiscoveryPing(const FTraceControlDiscoveryPing& Message, const TSharedRef<IMessageContext>& Context)
 {
-	if ((!Message.SessionId.IsValid() && !Message.InstanceId.IsValid()) || (Message.InstanceId == FApp::GetInstanceId() && Message.SessionId == FApp::GetSessionId()))
+	if ((!Message.SessionId.IsValid() && !Message.InstanceId.IsValid()) || (Message.InstanceId == FApp::GetInstanceId() || Message.SessionId == FApp::GetSessionId()))
 	{
 		const auto Response = FMessageEndpoint::MakeMessage<FTraceControlDiscovery>();
 		Response->SessionId = FApp::GetSessionId();
diff --git a/Engine/Source/Runtime/SessionServices/Private/TraceController.cpp b/Engine/Source/Runtime/SessionServices/Private/TraceController.cpp
--- a/Engine/Source/Runtime/SessionServices/Private/TraceController.cpp
+++ b/Engine/Source/Runtime/SessionServices/Private/TraceController.cpp
@@ -3,7 +3,6 @@
 #include "TraceController.h"
 
 #include "IMessageBus.h"
-#include "ISessionServicesModule.h"
 #include "MessageEndpoint.h"
 #include "MessageEndpointBuilder.h"
 #include "TraceControlMessages.h"
@@ -19,24 +18,10 @@ FTraceController::FTraceController(const TSharedRef<IMessageBus>& InMessageBus)
 		.Handling<FTraceControlChannelsDesc>(this, &FTraceController::OnChannelsDesc)
 		.Handling<FTraceControlChannelsStatus>(this, &FTraceController::OnChannelsStatus)
 		.NotificationHandling(FOnBusNotification::CreateRaw(this, &FTraceController::OnNotification));
-	
-	ISessionServicesModule& SessionServicesModule = FModuleManager::LoadModuleChecked<ISessionServicesModule>("SessionServices");
-	SessionManager = SessionServicesModule.GetSessionManager();
-	
-	if (SessionManager.IsValid())
-	{
-		//Note that we should not add ourselves as shared pointer to session manager,
-		//that will create a circular dependency
-		SessionManager->OnInstanceSelectionChanged().AddRaw(this, &FTraceController::OnInstanceSelectionChanged);
-	}
 }
 
 FTraceController::~FTraceController()
 {
-	if (SessionManager.IsValid())
-	{
-		SessionManager->OnInstanceSelectionChanged().RemoveAll(this);
-	}
 }
 
 void FTraceController::SendDiscoveryRequest(const FGuid& SessionId, const FGuid& InstanceId) const
@@ -54,18 +39,13 @@ void FTraceController::SendDiscoveryRequest()
 
 void FTraceController::SendStatusUpdateRequest()
 {
-	// The discovery pong contains the status, so don't send the request twice.
-	if (RediscoverSelectedSession() == false)
-	{
-		MessageEndpoint->Publish(FMessageEndpoint::MakeMessage<FTraceControlStatusPing>());
-	}
+	MessageEndpoint->Publish(FMessageEndpoint::MakeMessage<FTraceControlStatusPing>());
 }
 
 void FTraceController::SendChannelUpdateRequest()
 {
 	FReadScopeLock _(InstancesLock);
 	
-	RediscoverSelectedSession();
 	for (const auto& Instance : Instances)
 	{
 		const auto Message = FMessageEndpoint::MakeMessage<FTraceControlChannelsPing>();
@@ -125,10 +105,6 @@ void FTraceController::OnDiscoveryResponse(const FTraceControlDiscovery& Message
 
 	FTraceStatus::EUpdateType UpdateType = FTraceStatus::EUpdateType::ChannelsDesc | FTraceStatus::EUpdateType::ChannelsStatus | FTraceStatus::EUpdateType::Status;
 	StatusReceivedEvent.Broadcast(Instance->Status, UpdateType, Instance->Commands);
-	if (SelectedInstanceIds.Contains(Message.InstanceId))
-	{
-		SelectedSessionStatusReceivedEvent.Broadcast(Instance->Status, UpdateType, Instance->Commands);
-	}
 }
 
 void FTraceController::OnStatus(const FTraceControlStatus& Message, const TSharedRef<IMessageContext>& Context)
@@ -141,10 +117,6 @@ void FTraceController::OnStatus(const FTraceControlStatus& Message, const TShare
 		UpdateStatus(Message, Status);
 		
 		StatusReceivedEvent.Broadcast(Status, FTraceStatus::EUpdateType::Status, Instance->Commands);
-		if (SelectedInstanceIds.Contains(Status.InstanceId))
-		{
-			SelectedSessionStatusReceivedEvent.Broadcast(Status, FTraceStatus::EUpdateType::Status, Instance->Commands);
-		}
 	}
 }
 
@@ -178,10 +150,6 @@ void FTraceController::OnChannelsDesc(const FTraceControlChannelsDesc& Message,
 		NewChannel.bReadOnly = Message.ReadOnlyIds.Contains(Id);
 		
 		StatusReceivedEvent.Broadcast(Status, FTraceStatus::EUpdateType::ChannelsDesc, Instance->Commands);
-		if (SelectedInstanceIds.Contains(Status.InstanceId))
-		{
-			SelectedSessionStatusReceivedEvent.Broadcast(Status, FTraceStatus::EUpdateType::ChannelsDesc, Instance->Commands);
-		}
 	}
 
 	// Allow the commands instance to update its list of channels
@@ -201,10 +169,6 @@ void FTraceController::OnChannelsStatus(const FTraceControlChannelsStatus& Messa
 		}
 		
 		StatusReceivedEvent.Broadcast(Status, FTraceStatus::EUpdateType::ChannelsStatus, Instance->Commands);
-		if (SelectedInstanceIds.Contains(Status.InstanceId))
-		{
-			SelectedSessionStatusReceivedEvent.Broadcast(Status, FTraceStatus::EUpdateType::ChannelsStatus, Instance->Commands);
-		}
 	}
 }
 
@@ -227,54 +191,12 @@ void FTraceController::OnSettings(const FTraceControlSettings& Message, const TS
 		}
 		
 		StatusReceivedEvent.Broadcast(Status, FTraceStatus::EUpdateType::Settings, Instance->Commands);
-		if (SelectedInstanceIds.Contains(Status.InstanceId))
-		{
-			SelectedSessionStatusReceivedEvent.Broadcast(Status, FTraceStatus::EUpdateType::Settings, Instance->Commands);
-		}
 	}
 }
 
-bool FTraceController::HasAvailableSelectedInstance()
+bool FTraceController::HasAvailableInstance(const FGuid& InstanceId)
 {
-	// Return true if at least one of the selected instances is known
-	for (const FGuid& SelectedInstanceId : SelectedInstanceIds)
-	{
-		if (InstanceToAddress.Contains(SelectedInstanceId))
-		{
-			return true;
-		}
-	}
-	return false;
-}
-
-void FTraceController::OnInstanceSelectionChanged(const TSharedPtr<ISessionInstanceInfo>& InstanceInfo, bool bSelected)
-{
-	FReadScopeLock _(InstancesLock);
-	
-	if (bSelected)
-	{
-		SelectedInstanceIds.Add(InstanceInfo->GetInstanceId());
-	}
-	else
-	{
-		SelectedInstanceIds.Remove(InstanceInfo->GetInstanceId());
-	}
-
-	if (const auto Address = InstanceToAddress.Find(InstanceInfo->GetInstanceId()))
-	{
-		if (const auto Instance = Instances.Find(*Address))
-		{
-			FTraceStatus& Status = Instance->Status;
-			FTraceStatus::EUpdateType UpdateType = FTraceStatus::EUpdateType::ChannelsDesc | FTraceStatus::EUpdateType::ChannelsStatus | FTraceStatus::EUpdateType::Status;
-			SelectedSessionStatusReceivedEvent.Broadcast(Status, UpdateType, Instance->Commands);
-		}
-	}
-	else
-	{
-		SendDiscoveryRequest(InstanceInfo->GetOwnerSession()->GetSessionId(), InstanceInfo->GetInstanceId());
-	}
-
-	SessionSelectionChangedEvent.Broadcast();
+	return InstanceToAddress.Contains(InstanceId);
 }
 
 void FTraceController::UpdateStatus(const FTraceControlStatus& Message, FTraceStatus& Status)
@@ -295,21 +217,6 @@ void FTraceController::UpdateStatus(const FTraceControlStatus& Message, FTraceSt
 	Status.Stats.CacheWaste = Message.CacheWaste;
 }
 
-void FTraceController::WithSelectedInstances(FCallback Func)
-{
-	FReadScopeLock _(InstancesLock);
-	for (auto& SelectedInstance : SessionManager->GetSelectedInstances())
-	{
-		if(const auto& InstanceAddress = InstanceToAddress.Find(SelectedInstance->GetInstanceId()))
-		{
-			if(auto* Instance = Instances.Find(*InstanceAddress))
-			{
-				Func(Instance->Status, Instance->Commands);
-			}
-		}
-	}
-}
-
 void FTraceController::WithInstance(FGuid InstanceId, FCallback Func)
 {
 	FReadScopeLock _(InstancesLock);
@@ -322,20 +229,6 @@ void FTraceController::WithInstance(FGuid InstanceId, FCallback Func)
 	}
 }
 
-bool FTraceController::RediscoverSelectedSession() const
-{
-	auto& SelectedInstances = SessionManager->GetSelectedInstances();
-	for (auto& SelectedInstance : SelectedInstances)
-	{
-		if (!InstanceToAddress.Contains(SelectedInstance->GetInstanceId()))
-		{
-			SendDiscoveryRequest(SelectedInstance->GetOwnerSession()->GetSessionId(), SelectedInstance->GetInstanceId());
-			return true;
-		}
-	}
-	return false;
-}
-
 FTraceController::FTracingInstance::FTracingInstance(const TSharedRef<IMessageBus>& InMessageBus, FMessageAddress InService)
 	: Status()
 	, Commands(InMessageBus, InService)
diff --git a/Engine/Source/Runtime/SessionServices/Private/TraceController.h b/Engine/Source/Runtime/SessionServices/Private/TraceController.h
--- a/Engine/Source/Runtime/SessionServices/Private/TraceController.h
+++ b/Engine/Source/Runtime/SessionServices/Private/TraceController.h
@@ -3,7 +3,6 @@
 #pragma once
 
 #include "IMessageContext.h"
-#include "ISessionManager.h"
 #include "ITraceController.h"
 #include "MessageEndpoint.h"
 #include "TraceControllerCommands.h"
@@ -12,7 +11,6 @@
 struct FTraceControlSettings;
 struct FTraceControlStatus;
 struct FTraceControlDiscovery;
-class ISessionManager;
 class IMessageBus;
 class FMessageEndpoint;
 
@@ -31,8 +29,7 @@ private:
 	virtual void SendStatusUpdateRequest() override;
 	virtual void SendChannelUpdateRequest() override;
 	virtual void SendSettingsUpdateRequest() override;
-	virtual bool HasAvailableSelectedInstance() override;
-	virtual void WithSelectedInstances(FCallback Func) override;
+	virtual bool HasAvailableInstance(const FGuid& InstanceId) override;
 	virtual void WithInstance(FGuid InstanceId, FCallback Func) override;
 	
 	DECLARE_DERIVED_EVENT(FTraceController, ITraceController::FStatusRecievedEvent, FStatusRecievedEvent);
@@ -41,19 +38,6 @@ private:
 		return StatusReceivedEvent;
 	}
 
-	virtual FStatusRecievedEvent& OnSelectedSessionStatusReceived() override
-	{
-		return SelectedSessionStatusReceivedEvent;
-	}
-
-	DECLARE_DERIVED_EVENT(FTraceController, ITraceController::FSessionSelectionChanged, FSessionSelectionChanged);
-	virtual FSessionSelectionChanged& OnSessionSelectionChanged() override
-	{
-		return SessionSelectionChangedEvent;
-	}
-
-	virtual uint32 GetNumSelectedInstances() override { return SelectedInstanceIds.Num(); }
-
 	/* Message handlers */
 	void OnNotification(const FMessageBusNotification& MessageBusNotification);
 	void OnDiscoveryResponse(const FTraceControlDiscovery& Message, const TSharedRef<IMessageContext>& Context);
@@ -63,15 +47,6 @@ private:
 	void OnSettings(const FTraceControlSettings& Message, const TSharedRef<IMessageContext>& Context);
 	static void UpdateStatus(const FTraceControlStatus& Message, FTraceStatus& Status);
 
-	/* Events from SessionManager handlers */
-	void OnInstanceSelectionChanged(const TSharedPtr<class ISessionInstanceInfo>&, bool);
-
-	/* A selected instance can end up not discovered, either because the FTraceControlDiscoveryPong was lost
-	*  or because the selected session has been unregistered. Attempt to discover it again.
-	*  Returns true if a discovery ping was sent. 
-	*/
-	bool RediscoverSelectedSession() const;
-
 private:
 	
 	struct FTracingInstance
@@ -92,21 +67,12 @@ private:
 	/** Our own endpoint for messages */
 	TSharedPtr<FMessageEndpoint> MessageEndpoint;
 
-	/** Session manager used for selecting sessions */
-	TSharedPtr<ISessionManager> SessionManager;
-
 	/** Address of the runtime endpoint for trace controls */
 	FMessageAddress TraceControlAddress;
 
 	/** Event for status updates on any session */
 	FStatusRecievedEvent StatusReceivedEvent;
 
-	/** Event for status updates on a selected session */
-	FStatusRecievedEvent SelectedSessionStatusReceivedEvent;
-
-	/** Event that triggers when the session selection changes */
-	FSessionSelectionChanged SessionSelectionChangedEvent;
-
 	/** Lock to protect access to Instances list */
 	FRWLock InstancesLock;
 
@@ -115,7 +81,4 @@ private:
 
 	/** Secondary lookup from instance -> address */
 	TMap<FGuid, FMessageAddress> InstanceToAddress;
-
-	/** Currently selected instance */
-	TSet<FGuid> SelectedInstanceIds;
 };
diff --git a/Engine/Source/Runtime/SessionServices/Public/ITraceController.h b/Engine/Source/Runtime/SessionServices/Public/ITraceController.h
--- a/Engine/Source/Runtime/SessionServices/Public/ITraceController.h
+++ b/Engine/Source/Runtime/SessionServices/Public/ITraceController.h
@@ -141,13 +141,7 @@ ENUM_CLASS_FLAGS(FTraceStatus::EUpdateType);
  *		Commands.Send("localhost", "audio,audiomixer,bookmark,log");
  *		Commands.Bookmark("My remote bookmark");
  * });
- *
- * // We can also use the "selection" feature in session manager
- * TraceController->WithSelectedInstances([](const FTraceStatus& Status, ITraceControllerCommands& Commands){
- *		Commands.SnapshotSend("localhost");
- * });
  * ```
- * 
  */
 class ITraceController
 {
@@ -195,34 +189,11 @@ public:
 	virtual FStatusRecievedEvent& OnStatusReceived() = 0;
 
 	/**
-	 * Event triggered whenever the status for the selected instance (in session manager)
-	 * is updated. A reference to the status and what has changed is provided along with a 
-	 * structure to issue commands.
-	 */
-	virtual FStatusRecievedEvent& OnSelectedSessionStatusReceived() = 0;
-
-	/**
-	* Session selection changed event.
-	*/
-	DECLARE_EVENT(ITraceController, FSessionSelectionChanged);
-
-	/**
-	 * Event triggered when the session selection changes.
+	 * Return true if an instance with the provided Id exists and has been discovered.
 	 */
-	virtual FSessionSelectionChanged& OnSessionSelectionChanged() = 0;
-
-	/**
-	 * Return true if a selected instance exists and has been discovered.
-	 */
-	virtual bool HasAvailableSelectedInstance() = 0;
-
+	virtual bool HasAvailableInstance(const FGuid& InstanceId) = 0;
 
 	typedef TFunction<void(const FTraceStatus&, ITraceControllerCommands&)> FCallback;
-	/**
-	 * Execute a function for each instance selected in the session manager.
-	 * @param Func Functor to execute
-	 */
-	virtual void WithSelectedInstances(FCallback Func) = 0;
 
 	/**
 	 * Execute a function on a specific session.
@@ -231,11 +202,6 @@ public:
 	 */
 	virtual void WithInstance(FGuid InstanceId, FCallback Func) = 0;
 
-	/**
-	 * Returns the number of selected sessions.
-	 */
-	virtual uint32 GetNumSelectedInstances() = 0;
-
 	/**
 	 * Create an instance using a custom message bus
 	 * @param MessageBus Message bus the controller will use
