Add utility function to find out if a trace session has a certain channel

diff --git a/Engine/Source/Runtime/SessionServices/Public/ITraceController.h b/Engine/Source/Runtime/SessionServices/Public/ITraceController.h
--- a/Engine/Source/Runtime/SessionServices/Public/ITraceController.h
+++ b/Engine/Source/Runtime/SessionServices/Public/ITraceController.h
@@ -8,6 +8,7 @@
 #include "ITraceControllerCommands.h"
 #include "Misc/DateTime.h"
 #include "Misc/Guid.h"
+#include "Algo/Find.h"
 
 /**
  * Describes the state of a single remote instance with Trace.
@@ -119,6 +120,19 @@ struct FTraceStatus
 	FDateTime StatusTimestamp;
 	/** The status of the trace system. */
 	ETraceSystemStatus TraceSystemStatus;
+
+	bool HasChannel(FStringView ChannelName) const
+	{
+		for (const auto& ChannelPair : Channels)
+		{
+			FStringView Name (ChannelPair.Value.Name);
+			if (Name.Compare(ChannelName, ESearchCase::IgnoreCase) == 0)
+			{
+				return true;
+			}
+		}
+		return false;
+	}
 };
 
 ENUM_CLASS_FLAGS(FTraceStatus::EUpdateType);
