Fix crash in SessionInstanceInfo

diff --git a/Engine/Source/Runtime/SessionServices/Private/SessionInstanceInfo.cpp b/Engine/Source/Runtime/SessionServices/Private/SessionInstanceInfo.cpp
--- a/Engine/Source/Runtime/SessionServices/Private/SessionInstanceInfo.cpp
+++ b/Engine/Source/Runtime/SessionServices/Private/SessionInstanceInfo.cpp
@@ -23,6 +23,10 @@ FSessionInstanceInfo::FSessionInstanceInfo(const FGuid& InInstanceId, const TSha
 		.Handling<FSessionServiceLog>(this, &FSessionInstanceInfo::HandleSessionLogMessage);
 }
 
+FSessionInstanceInfo::~FSessionInstanceInfo()
+{
+	FMessageEndpoint::SafeRelease(MessageEndpoint);
+}
 
 /* FSessionInstanceInfo interface
  *****************************************************************************/
diff --git a/Engine/Source/Runtime/SessionServices/Private/SessionInstanceInfo.h b/Engine/Source/Runtime/SessionServices/Private/SessionInstanceInfo.h
--- a/Engine/Source/Runtime/SessionServices/Private/SessionInstanceInfo.h
+++ b/Engine/Source/Runtime/SessionServices/Private/SessionInstanceInfo.h
@@ -26,6 +26,9 @@ public:
 
 	/** Default constructor. */
 	FSessionInstanceInfo() { }
+	
+	/** Destructor. */
+	virtual ~FSessionInstanceInfo();
 
 	/**
 	 * Creates and initializes a new instance.
diff --git a/Engine/Source/Runtime/SessionServices/Private/SessionManager.cpp b/Engine/Source/Runtime/SessionServices/Private/SessionManager.cpp
--- a/Engine/Source/Runtime/SessionServices/Private/SessionManager.cpp
+++ b/Engine/Source/Runtime/SessionServices/Private/SessionManager.cpp
@@ -52,6 +52,7 @@ FSessionManager::FSessionManager(const TSharedRef<IMessageBus, ESPMode::ThreadSa
 FSessionManager::~FSessionManager()
 {
 	FTSTicker::GetCoreTicker().RemoveTicker(TickDelegateHandle);
+	FMessageEndpoint::SafeRelease(MessageEndpoint);
 }
 
 
