Add server timeout detection to Verse Insights.

diff --git a/Engine/Source/Runtime/SessionServices/Private/TraceController.cpp b/Engine/Source/Runtime/SessionServices/Private/TraceController.cpp
--- a/Engine/Source/Runtime/SessionServices/Private/TraceController.cpp
+++ b/Engine/Source/Runtime/SessionServices/Private/TraceController.cpp
@@ -102,6 +102,7 @@ void FTraceController::OnDiscoveryResponse(const FTraceControlDiscovery& Message
 	// same as trace session (maybe overridden by commandline)
 	Instance->Status.SessionId = Message.SessionId;
 	Instance->Status.InstanceId = Message.InstanceId;
+	Instance->LastMessageTimestamp = FDateTime::Now();
 
 	FTraceStatus::EUpdateType UpdateType = FTraceStatus::EUpdateType::ChannelsDesc | FTraceStatus::EUpdateType::ChannelsStatus | FTraceStatus::EUpdateType::Status;
 	StatusReceivedEvent.Broadcast(Instance->Status, UpdateType, Instance->Commands);
@@ -113,9 +114,11 @@ void FTraceController::OnStatus(const FTraceControlStatus& Message, const TShare
 	
 	if (FTracingInstance* Instance = Instances.Find(Context->GetSender()))
 	{
+		Instance->LastMessageTimestamp = FDateTime::Now();
+
 		FTraceStatus& Status = Instance->Status;
 		UpdateStatus(Message, Status);
-		
+
 		StatusReceivedEvent.Broadcast(Status, FTraceStatus::EUpdateType::Status, Instance->Commands);
 	}
 }
@@ -131,6 +134,8 @@ void FTraceController::OnChannelsDesc(const FTraceControlChannelsDesc& Message,
 		return;
 	}
 
+	Instance->LastMessageTimestamp = FDateTime::Now();
+
 	FTraceStatus& Status = Instance->Status;
 	check(Message.Channels.Num() == Message.Ids.Num() && Message.Channels.Num() == Message.Descriptions.Num());
 	const int32 Count = Message.Channels.Num();
@@ -162,6 +167,8 @@ void FTraceController::OnChannelsStatus(const FTraceControlChannelsStatus& Messa
 	
 	if (const auto Instance = Instances.Find(Context->GetSender()))
 	{
+		Instance->LastMessageTimestamp = FDateTime::Now();
+
 		FTraceStatus& Status = Instance->Status;
 		for (auto& [Id, Channel] : Status.Channels)
 		{
@@ -178,6 +185,8 @@ void FTraceController::OnSettings(const FTraceControlSettings& Message, const TS
 	
 	if (const auto Instance = Instances.Find(Context->GetSender()))
 	{
+		Instance->LastMessageTimestamp = FDateTime::Now();
+
 		FTraceStatus& Status = Instance->Status;
 		FTraceStatus::FSettings& Settings = Status.Settings;
 		Settings.bUseImportantCache = Message.bUseImportantCache;
@@ -199,6 +208,19 @@ bool FTraceController::HasAvailableInstance(const FGuid& InstanceId)
 	return InstanceToAddress.Contains(InstanceId);
 }
 
+bool FTraceController::IsActive(FGuid InstanceId, double TimeoutSec) const
+{
+	if (const auto& InstanceAddress = InstanceToAddress.Find(InstanceId))
+	{
+		if (auto* Instance = Instances.Find(*InstanceAddress))
+		{
+			return (FDateTime::Now() - Instance->LastMessageTimestamp).GetTotalSeconds() <= TimeoutSec;
+		}
+	}
+
+	return false;
+}
+
 void FTraceController::UpdateStatus(const FTraceControlStatus& Message, FTraceStatus& Status)
 {
 	Status.TraceSystemStatus = static_cast<FTraceStatus::ETraceSystemStatus>(Message.TraceSystemStatus);
diff --git a/Engine/Source/Runtime/SessionServices/Private/TraceController.h b/Engine/Source/Runtime/SessionServices/Private/TraceController.h
--- a/Engine/Source/Runtime/SessionServices/Private/TraceController.h
+++ b/Engine/Source/Runtime/SessionServices/Private/TraceController.h
@@ -30,6 +30,7 @@ private:
 	virtual void SendChannelUpdateRequest() override;
 	virtual void SendSettingsUpdateRequest() override;
 	virtual bool HasAvailableInstance(const FGuid& InstanceId) override;
+	virtual bool IsActive(FGuid InstanceId, double TimeoutSec = 5.0f) const override;
 	virtual void WithInstance(FGuid InstanceId, FCallback Func) override;
 	
 	DECLARE_DERIVED_EVENT(FTraceController, ITraceController::FStatusRecievedEvent, FStatusRecievedEvent);
@@ -53,6 +54,7 @@ private:
 	{
 		FTraceStatus Status;
 		FTraceControllerCommands Commands;
+		FDateTime LastMessageTimestamp;
 
 		FTracingInstance(const TSharedRef<IMessageBus>& InMessageBus, FMessageAddress Service);
 		FTracingInstance() = delete;
diff --git a/Engine/Source/Runtime/SessionServices/Public/ITraceController.h b/Engine/Source/Runtime/SessionServices/Public/ITraceController.h
--- a/Engine/Source/Runtime/SessionServices/Public/ITraceController.h
+++ b/Engine/Source/Runtime/SessionServices/Public/ITraceController.h
@@ -208,6 +208,13 @@ public:
 	 */
 	virtual bool HasAvailableInstance(const FGuid& InstanceId) = 0;
 
+	/**
+	 * Returns true if the last message received for a specific instance is newer than TimeoutSec.
+	 * @param InstanceId The id of the instance to check.
+	 * @param TimeoutSec The maximum number of seconds that must have passed since the last message for it to be considered active.
+	 */
+	virtual bool IsActive(FGuid InstanceId, double TimeoutSec = 5.0f) const = 0;
+
 	typedef TFunction<void(const FTraceStatus&, ITraceControllerCommands&)> FCallback;
 
 	/**
