[Insights] Trace controller support for secure connections

- Added support for controlling secure connections.

#Synced-CL 49918891

diff --git a/Engine/Source/Runtime/SessionServices/Private/TraceController.cpp b/Engine/Source/Runtime/SessionServices/Private/TraceController.cpp
--- a/Engine/Source/Runtime/SessionServices/Private/TraceController.cpp
+++ b/Engine/Source/Runtime/SessionServices/Private/TraceController.cpp
@@ -12,12 +12,12 @@ FTraceController::FTraceController(const TSharedRef<IMessageBus>& InMessageBus)
 	: MessageBus(InMessageBus)
 {
 	MessageEndpoint = FMessageEndpoint::Builder("FTraceController", InMessageBus)
-		.Handling<FTraceControlDiscovery>(this, &FTraceController::OnDiscoveryResponse)
-		.Handling<FTraceControlStatus>(this, &FTraceController::OnStatus)
-		.Handling<FTraceControlSettings>(this, &FTraceController::OnSettings)
-		.Handling<FTraceControlChannelsDesc>(this, &FTraceController::OnChannelsDesc)
-		.Handling<FTraceControlChannelsStatus>(this, &FTraceController::OnChannelsStatus)
-		.NotificationHandling(FOnBusNotification::CreateRaw(this, &FTraceController::OnNotification));
+	  .Handling<FTraceControlDiscovery>(this, &FTraceController::OnDiscoveryResponse)
+	  .Handling<FTraceControlStatus>(this, &FTraceController::OnStatus)
+	  .Handling<FTraceControlSettings>(this, &FTraceController::OnSettings)
+	  .Handling<FTraceControlChannelsDesc>(this, &FTraceController::OnChannelsDesc)
+	  .Handling<FTraceControlChannelsStatus>(this, &FTraceController::OnChannelsStatus)
+	  .NotificationHandling(FOnBusNotification::CreateRaw(this, &FTraceController::OnNotification));
 }
 
 FTraceController::~FTraceController()
@@ -31,6 +31,7 @@ void FTraceController::SendDiscoveryRequest(const FGuid& SessionId, const FGuid&
 	Message->InstanceId = InstanceId;
 	MessageEndpoint->Publish<FTraceControlDiscoveryPing>(Message);
 }
+
 void FTraceController::SendDiscoveryRequest()
 {
 	const auto Message = FMessageEndpoint::MakeMessage<FTraceControlDiscoveryPing>();
@@ -92,7 +93,7 @@ void FTraceController::OnDiscoveryResponse(const FTraceControlDiscovery& Message
 	if (!Instance)
 	{
 		// Create a new instance with default status
-		Instance = &Instances.Emplace(Context->GetSender(), {Bus.ToSharedRef(), Context->GetSender()});
+		Instance = &Instances.Emplace(Context->GetSender(), {Bus.ToSharedRef(), Context->GetSender(), this});
 		InstanceToAddress.Add(Message.InstanceId, Context->GetSender());
 	}
 	
@@ -203,6 +204,27 @@ void FTraceController::OnSettings(const FTraceControlSettings& Message, const TS
 	}
 }
 
+void FTraceController::OnSecureResponse(const FTraceControlSecureResponse& Message, const TSharedRef<IMessageContext>& Context)
+{
+	FWriteScopeLock _(InstancesLock);
+	
+	if (const auto Instance = Instances.Find(Context->GetSender()))
+	{
+		FSecureConnection Connection;
+		Connection.Port = Message.Port;
+		Connection.AuthorizationToken = Message.AuthorizationToken;
+		if (Message.Mode == ETraceSecureMode::SelfSignedCertAndAuthToken)
+		{
+			Connection.CertificateAuthority = Message.CertificateAuthority;
+		}
+		else if (Message.Mode == ETraceSecureMode::AuthToken)
+		{
+			Connection.Identity = Message.Identity;
+		}
+		SecureResponse.Broadcast(Instance->Status, Connection);
+	}
+}
+
 bool FTraceController::HasAvailableInstance(const FGuid& InstanceId)
 {
 	return InstanceToAddress.Contains(InstanceId);
@@ -251,13 +273,13 @@ void FTraceController::WithInstance(FGuid InstanceId, FCallback Func)
 	}
 }
 
-FTraceController::FTracingInstance::FTracingInstance(const TSharedRef<IMessageBus>& InMessageBus, FMessageAddress InService)
+FTraceController::FTracingInstance::FTracingInstance(const TSharedRef<IMessageBus>& InMessageBus, FMessageAddress InService, FTraceController* Controller)
 	: Status()
-	, Commands(InMessageBus, InService)
+	, Commands(InMessageBus, InService, Controller)
 {
 }
 
 TSharedPtr<ITraceController> ITraceController::Create(TSharedPtr<IMessageBus>& InBus)
 {
 	return MakeShareable(new FTraceController(InBus.ToSharedRef()));
-}
\ No newline at end of file
+}
diff --git a/Engine/Source/Runtime/SessionServices/Private/TraceController.h b/Engine/Source/Runtime/SessionServices/Private/TraceController.h
--- a/Engine/Source/Runtime/SessionServices/Private/TraceController.h
+++ b/Engine/Source/Runtime/SessionServices/Private/TraceController.h
@@ -33,12 +33,18 @@ private:
 	virtual bool IsActive(FGuid InstanceId, double TimeoutSec = 5.0f) const override;
 	virtual void WithInstance(FGuid InstanceId, FCallback Func) override;
 	
-	DECLARE_DERIVED_EVENT(FTraceController, ITraceController::FStatusRecievedEvent, FStatusRecievedEvent);
-	virtual FStatusRecievedEvent& OnStatusReceived() override
+	DECLARE_DERIVED_EVENT(FTraceController, ITraceController::FStatusReceivedEvent, FStatusReceivedEvent);
+	virtual FStatusReceivedEvent& OnStatusReceived() override
 	{
 		return StatusReceivedEvent;
 	}
 
+	DECLARE_DERIVED_EVENT(FTraceController, ITraceController::FSecureResponseEvent, FSecureResponseEvent);
+	virtual FSecureResponseEvent& OnSecureResponse() override
+	{
+		return SecureResponse;
+	}
+
 	/* Message handlers */
 	void OnNotification(const FMessageBusNotification& MessageBusNotification);
 	void OnDiscoveryResponse(const FTraceControlDiscovery& Message, const TSharedRef<IMessageContext>& Context);
@@ -47,6 +53,7 @@ private:
 	void OnChannelsStatus(const FTraceControlChannelsStatus& Message, const TSharedRef<IMessageContext>& Context);
 	void OnSettings(const FTraceControlSettings& Message, const TSharedRef<IMessageContext>& Context);
 	static void UpdateStatus(const FTraceControlStatus& Message, FTraceStatus& Status);
+	void OnSecureResponse(const FTraceControlSecureResponse& Message, const TSharedRef<IMessageContext>& Context);
 
 private:
 	
@@ -56,7 +63,7 @@ private:
 		FTraceControllerCommands Commands;
 		FDateTime LastMessageTimestamp;
 
-		FTracingInstance(const TSharedRef<IMessageBus>& InMessageBus, FMessageAddress Service);
+		FTracingInstance(const TSharedRef<IMessageBus>& InMessageBus, FMessageAddress Service, FTraceController* Controller);
 		FTracingInstance() = delete;
 	};
 
@@ -73,7 +80,10 @@ private:
 	FMessageAddress TraceControlAddress;
 
 	/** Event for status updates on any session */
-	FStatusRecievedEvent StatusReceivedEvent;
+	FStatusReceivedEvent StatusReceivedEvent;
+
+	/** Event for secure connection responses */
+	FSecureResponseEvent SecureResponse;
 
 	/** Lock to protect access to Instances list */
 	FRWLock InstancesLock;
@@ -83,4 +93,6 @@ private:
 
 	/** Secondary lookup from instance -> address */
 	TMap<FGuid, FMessageAddress> InstanceToAddress;
+
+	friend class FTraceControllerCommands;
 };
diff --git a/Engine/Source/Runtime/SessionServices/Private/TraceControllerCommands.cpp b/Engine/Source/Runtime/SessionServices/Private/TraceControllerCommands.cpp
--- a/Engine/Source/Runtime/SessionServices/Private/TraceControllerCommands.cpp
+++ b/Engine/Source/Runtime/SessionServices/Private/TraceControllerCommands.cpp
@@ -6,6 +6,7 @@
 #include "MessageEndpoint.h"
 #include "MessageEndpointBuilder.h"
 #include "TraceControlMessages.h"
+#include "TraceController.h"
 
 
 //------------------------------------------------------------------------------------------
@@ -41,12 +42,13 @@ TArray<uint32> ResolveChannelNames(TConstArrayView<StringType> ChannelNames, con
 // FTraceControllerCommands
 //------------------------------------------------------------------------------------------
 
-FTraceControllerCommands::FTraceControllerCommands(const TSharedPtr<IMessageBus>& InMessageBus, FMessageAddress Service)
+FTraceControllerCommands::FTraceControllerCommands(const TSharedPtr<IMessageBus>& InMessageBus, FMessageAddress Service, class FTraceController* Controller)
 	: ServiceAddress(Service)
 {
 	TStringBuilder<128> EndpointName;
 	EndpointName << TEXT("FTraceControllerCommands_") << Service.ToString();
-	MessageEndpoint = FMessageEndpoint::Builder(EndpointName.ToString(), InMessageBus.ToSharedRef());
+	MessageEndpoint = FMessageEndpoint::Builder(EndpointName.ToString(), InMessageBus.ToSharedRef())
+		.Handling<FTraceControlSecureResponse>(Controller, &FTraceController::OnSecureResponse);
 }
 
 FTraceControllerCommands::~FTraceControllerCommands()
@@ -166,3 +168,14 @@ void FTraceControllerCommands::OnChannelsDesc(const FTraceControlChannelsDesc& M
 		}
 	}
 }
+
+void FTraceControllerCommands::SendSecure(FStringView Channels, bool bExcludeTail)
+{
+	//todo: If we had access to the target we could return here is secure connections
+	//are not supported.
+	const auto Message = FMessageEndpoint::MakeMessage<FTraceControlSecure>();
+	Message->Channels = Channels;
+	Message->bExcludeTail = bExcludeTail;
+	MessageEndpoint->Send(Message, ServiceAddress);
+}
+
diff --git a/Engine/Source/Runtime/SessionServices/Private/TraceControllerCommands.h b/Engine/Source/Runtime/SessionServices/Private/TraceControllerCommands.h
--- a/Engine/Source/Runtime/SessionServices/Private/TraceControllerCommands.h
+++ b/Engine/Source/Runtime/SessionServices/Private/TraceControllerCommands.h
@@ -14,7 +14,7 @@ struct FTraceControlChannelsDesc;
 class FTraceControllerCommands : public ITraceControllerCommands
 {
 public:
-	FTraceControllerCommands(const TSharedPtr<IMessageBus>& MessageBus, FMessageAddress Service);
+	FTraceControllerCommands(const TSharedPtr<IMessageBus>& MessageBus, FMessageAddress Service, class FTraceController* Controller);
 	virtual ~FTraceControllerCommands() override;
 
 	void OnChannelsDesc(const FTraceControlChannelsDesc& Message);
@@ -33,6 +33,7 @@ private:
 	virtual void Bookmark(FStringView Label) override;
 	virtual void Screenshot(FStringView Name, bool bShowUI) override;
 	virtual void SetStatNamedEventsEnabled(bool bEnabled) override;
+	virtual void SendSecure(FStringView Channels, bool bExcludeTail) override;
 
 private:
 	/** Address of the remote trace service to control */
diff --git a/Engine/Source/Runtime/SessionServices/Public/ITraceController.h b/Engine/Source/Runtime/SessionServices/Public/ITraceController.h
--- a/Engine/Source/Runtime/SessionServices/Public/ITraceController.h
+++ b/Engine/Source/Runtime/SessionServices/Public/ITraceController.h
@@ -138,6 +138,13 @@ struct FTraceStatus
 
 ENUM_CLASS_FLAGS(FTraceStatus::EUpdateType);
 
+struct FSecureConnection
+{
+	uint16 Port;
+	TArray<uint8> AuthorizationToken;
+	TArray<uint8> CertificateAuthority;
+	FString Identity;
+};
 
 /**
  * Interface for monitoring known trace service instances. ITraceController is a singleton which tracks
@@ -194,14 +201,22 @@ public:
 	/**
 	 * Generic event for updates of status
 	 */
-	DECLARE_EVENT_ThreeParams(ITraceController, FStatusRecievedEvent, const FTraceStatus&, FTraceStatus::EUpdateType, ITraceControllerCommands&);
+	DECLARE_EVENT_ThreeParams(ITraceController, FStatusReceivedEvent, const FTraceStatus&, FTraceStatus::EUpdateType, ITraceControllerCommands&);
 
 	/**
 	 * Event triggered whenever status is updated for any known session. 
 	 * A reference to the status and what has changed is provided along with a 
 	 * structure to issue commands.
 	 */
-	virtual FStatusRecievedEvent& OnStatusReceived() = 0;
+	virtual FStatusReceivedEvent& OnStatusReceived() = 0;
+
+
+	DECLARE_EVENT_TwoParams(ITraceController, FSecureResponseEvent, const FTraceStatus&, const FSecureConnection&)
+
+	/**
+	 * Event triggered when a response to a secure connection request has been made.
+	 */
+	virtual FSecureResponseEvent& OnSecureResponse() = 0;
 
 	/**
 	 * Return true if an instance with the provided Id exists and has been discovered.
diff --git a/Engine/Source/Runtime/SessionServices/Public/ITraceControllerCommands.h b/Engine/Source/Runtime/SessionServices/Public/ITraceControllerCommands.h
--- a/Engine/Source/Runtime/SessionServices/Public/ITraceControllerCommands.h
+++ b/Engine/Source/Runtime/SessionServices/Public/ITraceControllerCommands.h
@@ -45,6 +45,15 @@ public:
 	 */
 	virtual void File(FStringView File, FStringView Channels, bool bExcludeTail = false, bool bTruncateFile = false) = 0;
 
+
+	/**
+	 * Start a trace on selected instance using a secure connection. The initiator is required
+	 * to connect to the port using the options specified in the response.
+	 * @param Channels Comma separated list of channels to enable
+	 * @param bExcludeTail If the tail (circular buffer of recent events) should be included
+	 */
+	virtual void SendSecure(FStringView Channels, bool bExcludeTail = false) = 0;
+
 	/**
 	 * On selected instances, make a snapshot of the tail (circular buffer of recent events)
 	 * and send to the provided host.
