Unicast trace status and channel status messages

#Synced-CL 50044409

diff --git a/Engine/Source/Developer/TraceTools/Private/Widgets/STraceControl.cpp b/Engine/Source/Developer/TraceTools/Private/Widgets/STraceControl.cpp
--- a/Engine/Source/Developer/TraceTools/Private/Widgets/STraceControl.cpp
+++ b/Engine/Source/Developer/TraceTools/Private/Widgets/STraceControl.cpp
@@ -39,9 +39,6 @@ void STraceControl::Construct(const FArguments& InArgs, TSharedPtr<ITraceControl
 
 	UICommandList = MakeShareable(new FUICommandList);
 
-	TraceController->SendStatusUpdateRequest();
-	TraceController->SendChannelUpdateRequest();
-
 	SessionFilterService = MakeShareable(new FSessionTraceControllerFilterService(InTraceController));
 
 	ChildSlot
@@ -84,14 +81,47 @@ void STraceControl::Construct(const FArguments& InArgs, TSharedPtr<ITraceControl
 	}
 }
 
+void STraceControl::Tick(const FGeometry& AllottedGeometry, const double InCurrentTime, const float InDeltaTime)
+{
+	AccumulatedTime += InDeltaTime;
+
+	constexpr double UpdateTime = 1.0f;
+
+	if (InstanceId.IsValid() && AccumulatedTime > UpdateTime)
+	{
+		TraceController->WithInstance(InstanceId, [](const FTraceStatus& Status, ITraceControllerCommands& Commands)
+		{
+			Commands.SendStatusUpdateRequest();
+			Commands.SendChannelUpdateRequest();
+		});
+
+		if (!SessionFilterService->HasSettings())
+		{
+			TraceController->WithInstance(InstanceId, [](const FTraceStatus& Status, ITraceControllerCommands& Commands)
+			{
+				Commands.SendSettingsUpdateRequest();
+			});
+		}
+
+		AccumulatedTime = 0.0f;
+	}
+}
+
 void STraceControl::SetInstanceId(const FGuid& Id)
 {
+	InstanceId = Id;
 	Toolbar->SetInstanceId(Id);
 	SessionFilterService->SetInstanceId(Id);
 
-	if (Id.IsValid())
+	if (InstanceId.IsValid())
 	{
 		TraceController->SendDiscoveryRequest(FGuid(), Id);
+
+		TraceController->WithInstance(InstanceId, [](const FTraceStatus& Status, ITraceControllerCommands& Commands)
+		{
+			Commands.SendStatusUpdateRequest();
+			Commands.SendChannelUpdateRequest();
+		});
 	}
 
 #if !WITH_EDITOR
diff --git a/Engine/Source/Developer/TraceTools/Private/Widgets/STraceControl.h b/Engine/Source/Developer/TraceTools/Private/Widgets/STraceControl.h
--- a/Engine/Source/Developer/TraceTools/Private/Widgets/STraceControl.h
+++ b/Engine/Source/Developer/TraceTools/Private/Widgets/STraceControl.h
@@ -40,6 +40,8 @@ public:
 	/** Sets the InstanceId to control. Supports invalid guid value for disabled state. */
 	void SetInstanceId(const FGuid& Id);
 
+	virtual void Tick(const FGeometry& AllottedGeometry, const double InCurrentTime, const float InDeltaTime) override;
+
 protected:
 	void OnInstanceSelectionChanged(const TSharedPtr<class ISessionInstanceInfo>&, bool);
 
@@ -63,6 +65,8 @@ private:
 	TSet<FGuid> SelectedSessionsIds;
 
 	bool bAutoDetectSelectedSession;
+
+	double AccumulatedTime = 0.0f;
 };
 
 } // namespace UE::TraceTools
diff --git a/Engine/Source/Developer/TraceTools/Private/Widgets/STraceDataFilterWidget.cpp b/Engine/Source/Developer/TraceTools/Private/Widgets/STraceDataFilterWidget.cpp
--- a/Engine/Source/Developer/TraceTools/Private/Widgets/STraceDataFilterWidget.cpp
+++ b/Engine/Source/Developer/TraceTools/Private/Widgets/STraceDataFilterWidget.cpp
@@ -547,23 +547,6 @@ void STraceDataFilterWidget::Tick(const FGeometry& AllottedGeometry, const doubl
 			bNeedsListRefresh = false;
 		}
 	}
-
-	AccumulatedTime += InDeltaTime;
-
-	constexpr double UpdateTime = 1.0f;
-
-	if (AccumulatedTime > UpdateTime)
-	{
-		TraceController->SendStatusUpdateRequest();
-		TraceController->SendChannelUpdateRequest();
-
-		if (!SessionFilterService->HasSettings())
-		{
-			TraceController->SendSettingsUpdateRequest();
-		}
-
-		AccumulatedTime = 0.0;
-	}
 }
 
 void STraceDataFilterWidget::OnSessionSelectionChanged()
diff --git a/Engine/Source/Developer/TraceTools/Private/Widgets/STraceDataFilterWidget.h b/Engine/Source/Developer/TraceTools/Private/Widgets/STraceDataFilterWidget.h
--- a/Engine/Source/Developer/TraceTools/Private/Widgets/STraceDataFilterWidget.h
+++ b/Engine/Source/Developer/TraceTools/Private/Widgets/STraceDataFilterWidget.h
@@ -127,8 +127,6 @@ protected:
 
 	TSharedPtr<ITraceController> TraceController;
 
-	double AccumulatedTime = 0.0f;
-
 	bool bHasChannelData = false;
 	bool bHasSettings = false;
 
diff --git a/Engine/Source/Runtime/SessionServices/Private/TraceControllerCommands.cpp b/Engine/Source/Runtime/SessionServices/Private/TraceControllerCommands.cpp
--- a/Engine/Source/Runtime/SessionServices/Private/TraceControllerCommands.cpp
+++ b/Engine/Source/Runtime/SessionServices/Private/TraceControllerCommands.cpp
@@ -5,6 +5,7 @@
 #include "Hash/xxhash.h"
 #include "MessageEndpoint.h"
 #include "MessageEndpointBuilder.h"
+#include "TraceController.h"
 #include "TraceControlMessages.h"
 #include "TraceController.h"
 
@@ -48,7 +49,11 @@ FTraceControllerCommands::FTraceControllerCommands(const TSharedPtr<IMessageBus>
 	TStringBuilder<128> EndpointName;
 	EndpointName << TEXT("FTraceControllerCommands_") << Service.ToString();
 	MessageEndpoint = FMessageEndpoint::Builder(EndpointName.ToString(), InMessageBus.ToSharedRef())
-		.Handling<FTraceControlSecureResponse>(Controller, &FTraceController::OnSecureResponse);
+		.Handling<FTraceControlSecureResponse>(Controller, &FTraceController::OnSecureResponse)
+		.Handling<FTraceControlStatus>(Controller, &FTraceController::OnStatus)
+		.Handling<FTraceControlChannelsDesc>(Controller, &FTraceController::OnChannelsDesc)
+		.Handling<FTraceControlChannelsStatus>(Controller, &FTraceController::OnChannelsStatus)
+		.Handling<FTraceControlSettings>(Controller, &FTraceController::OnSettings);
 }
 
 FTraceControllerCommands::~FTraceControllerCommands()
@@ -148,9 +153,31 @@ void FTraceControllerCommands::SetStatNamedEventsEnabled(bool bEnabled)
 	MessageEndpoint->Send(Message, ServiceAddress);
 }
 
+void FTraceControllerCommands::SendStatusUpdateRequest()
+{
+	const auto Message = FMessageEndpoint::MakeMessage<FTraceControlStatusPing>();
+
+	MessageEndpoint->Send(Message, ServiceAddress);
+}
+
+void FTraceControllerCommands::SendChannelUpdateRequest()
+{
+	const auto Message = FMessageEndpoint::MakeMessage<FTraceControlChannelsPing>();
+	Message->KnownChannelCount = ChannelCount;
+
+	MessageEndpoint->Send(Message, ServiceAddress);
+}
+
+void FTraceControllerCommands::SendSettingsUpdateRequest()
+{
+	const auto Message = FMessageEndpoint::MakeMessage<FTraceControlSettingsPing>();
+
+	MessageEndpoint->Send(Message, ServiceAddress);
+}
+
 void FTraceControllerCommands::OnChannelsDesc(const FTraceControlChannelsDesc& Message)
 {
-	const int32 ChannelCount = Message.Channels.Num();
+	ChannelCount = Message.Channels.Num();
 	if (ChannelCount != SettableChannels.Num())
 	{
 		for (int32 ChannelIdx = 0; ChannelIdx < ChannelCount; ++ChannelIdx)
diff --git a/Engine/Source/Runtime/SessionServices/Private/TraceControllerCommands.h b/Engine/Source/Runtime/SessionServices/Private/TraceControllerCommands.h
--- a/Engine/Source/Runtime/SessionServices/Private/TraceControllerCommands.h
+++ b/Engine/Source/Runtime/SessionServices/Private/TraceControllerCommands.h
@@ -7,6 +7,7 @@
 class IMessageBus;
 class FMessageEndpoint;
 struct FTraceControlChannelsDesc;
+class FTraceController;
 
 /**
  * Logic for sending commands to a specific remote trace service.
@@ -34,6 +35,9 @@ private:
 	virtual void Screenshot(FStringView Name, bool bShowUI) override;
 	virtual void SetStatNamedEventsEnabled(bool bEnabled) override;
 	virtual void SendSecure(FStringView Channels, bool bExcludeTail) override;
+	virtual void SendStatusUpdateRequest() override;
+	virtual void SendChannelUpdateRequest() override;
+	virtual void SendSettingsUpdateRequest() override;
 
 private:
 	/** Address of the remote trace service to control */
@@ -44,4 +48,7 @@ private:
 
 	/** List of channels (hash of their name) that is possible to set (not read-only). */
 	TMap<uint64, uint32> SettableChannels;
+
+	/** The current known channel count. */
+	int32 ChannelCount = 0;
 };
diff --git a/Engine/Source/Runtime/SessionServices/Public/ITraceController.h b/Engine/Source/Runtime/SessionServices/Public/ITraceController.h
--- a/Engine/Source/Runtime/SessionServices/Public/ITraceController.h
+++ b/Engine/Source/Runtime/SessionServices/Public/ITraceController.h
@@ -186,16 +186,19 @@ public:
 	/**
 	 * Request update of the status from all sessions and instances.
 	 */
+	UE_DEPRECATED(5.8, "Use the version from ITraceControllerCommands as this one can cause high traffic when used used in large networks.")
 	virtual void SendStatusUpdateRequest() = 0;
 
 	/**
 	 * Request update of channel status from all sessions and instances.
 	 */
+	UE_DEPRECATED(5.8, "Use the version from ITraceControllerCommands as this one can cause high traffic when used used in large networks.")
 	virtual void SendChannelUpdateRequest() = 0;
 
 	/**
 	 * Request update of setting state from all sessions and instances.
 	 */
+	UE_DEPRECATED(5.8, "Use the version from ITraceControllerCommands as this one can cause high traffic when used used in large networks.")
 	virtual void SendSettingsUpdateRequest() = 0;
 
 	/**
diff --git a/Engine/Source/Runtime/SessionServices/Public/ITraceControllerCommands.h b/Engine/Source/Runtime/SessionServices/Public/ITraceControllerCommands.h
--- a/Engine/Source/Runtime/SessionServices/Public/ITraceControllerCommands.h
+++ b/Engine/Source/Runtime/SessionServices/Public/ITraceControllerCommands.h
@@ -101,5 +101,20 @@ public:
 	 * @param bEnabled The value to assign to the StatNamedEvents flag.
 	 */
 	virtual void SetStatNamedEventsEnabled(bool bEnabled) = 0;
+
+	/**
+	 * Request update of the status.
+	 */
+	virtual void SendStatusUpdateRequest() = 0;
+
+	/**
+	 * Request update of channel status.
+	 */
+	virtual void SendChannelUpdateRequest() = 0;
+
+	/**
+	 * Request update of setting state.
+	*/
+	virtual void SendSettingsUpdateRequest() = 0;
 };
 
