Temporary disable deprecation some warnings in Trace Controllers because of integration issues.

#Synced-CL 50044409

diff --git a/Engine/Source/Runtime/SessionServices/Public/ITraceController.h b/Engine/Source/Runtime/SessionServices/Public/ITraceController.h
--- a/Engine/Source/Runtime/SessionServices/Public/ITraceController.h
+++ b/Engine/Source/Runtime/SessionServices/Public/ITraceController.h
@@ -186,19 +186,19 @@ public:
 	/**
 	 * Request update of the status from all sessions and instances.
 	 */
-	UE_DEPRECATED(5.8, "Use the version from ITraceControllerCommands as this one can cause high traffic when used used in large networks.")
+	// UE_DEPRECATED(5.8, "Use the version from ITraceControllerCommands as this one can cause high traffic when used used in large networks.")
 	virtual void SendStatusUpdateRequest() = 0;
 
 	/**
 	 * Request update of channel status from all sessions and instances.
 	 */
-	UE_DEPRECATED(5.8, "Use the version from ITraceControllerCommands as this one can cause high traffic when used used in large networks.")
+	// UE_DEPRECATED(5.8, "Use the version from ITraceControllerCommands as this one can cause high traffic when used used in large networks.")
 	virtual void SendChannelUpdateRequest() = 0;
 
 	/**
 	 * Request update of setting state from all sessions and instances.
 	 */
-	UE_DEPRECATED(5.8, "Use the version from ITraceControllerCommands as this one can cause high traffic when used used in large networks.")
+	// UE_DEPRECATED(5.8, "Use the version from ITraceControllerCommands as this one can cause high traffic when used used in large networks.")
 	virtual void SendSettingsUpdateRequest() = 0;
 
 	/**
