[Insights] Add host to controller secure status struct.

#Synced-CL 50345636

diff --git a/Engine/Source/Runtime/SessionServices/Private/TraceController.cpp b/Engine/Source/Runtime/SessionServices/Private/TraceController.cpp
--- a/Engine/Source/Runtime/SessionServices/Private/TraceController.cpp
+++ b/Engine/Source/Runtime/SessionServices/Private/TraceController.cpp
@@ -211,6 +211,7 @@ void FTraceController::OnSecureResponse(const FTraceControlSecureResponse& Messa
 	if (const auto Instance = Instances.Find(Context->GetSender()))
 	{
 		FSecureConnection Connection;
+		Connection.Host = Message.Host;
 		Connection.Port = Message.Port;
 		Connection.AuthorizationToken = Message.AuthorizationToken;
 		if (Message.Mode == ETraceSecureMode::SelfSignedCertAndAuthToken)
diff --git a/Engine/Source/Runtime/SessionServices/Public/ITraceController.h b/Engine/Source/Runtime/SessionServices/Public/ITraceController.h
--- a/Engine/Source/Runtime/SessionServices/Public/ITraceController.h
+++ b/Engine/Source/Runtime/SessionServices/Public/ITraceController.h
@@ -141,6 +141,7 @@ ENUM_CLASS_FLAGS(FTraceStatus::EUpdateType);
 
 struct FSecureConnection
 {
+	FString Host;
 	uint16 Port;
 	TArray<uint8> AuthorizationToken;
 	TArray<uint8> CertificateAuthority;
