[Insights] Expose capabilities

* Expose capabilities of a target session through a querry function.
* Added locks to some functions that accessed instance data.

#Synced-CL 50431656

diff --git a/Engine/Source/Runtime/SessionServices/Private/TraceController.cpp b/Engine/Source/Runtime/SessionServices/Private/TraceController.cpp
--- a/Engine/Source/Runtime/SessionServices/Private/TraceController.cpp
+++ b/Engine/Source/Runtime/SessionServices/Private/TraceController.cpp
@@ -193,6 +193,9 @@ void FTraceController::OnSettings(const FTraceControlSettings& Message, const TS
 		Settings.bUseImportantCache = Message.bUseImportantCache;
 		Settings.bUseWorkerThread = Message.bUseWorkerThread;
 		Settings.TailSizeBytes = Message.TailSizeBytes;
+		Algo::Transform(Message.Capabilities, Settings.Capabilities, [](const FTraceCapability& In) { 
+				return FTraceStatus::FCapability { In.Name }; 
+		});
 
 		Settings.ChannelPresets.Empty();
 		for (const FTraceChannelPreset& Preset : Message.ChannelPresets)
@@ -244,6 +247,25 @@ bool FTraceController::IsActive(FGuid InstanceId, double TimeoutSec) const
 	return false;
 }
 
+bool FTraceController::HasCapability(FGuid InstanceId, FStringView Name) 
+{
+	FReadScopeLock _(InstancesLock);
+	if (const auto& InstanceAddress = InstanceToAddress.Find(InstanceId))
+	{
+		if (auto* Instance = Instances.Find(*InstanceAddress))
+		{
+			for (const auto& Capability : Instance->Status.Settings.Capabilities)
+			{
+				if (Name.Equals(Capability.Name, ESearchCase::IgnoreCase))
+				{
+					return true;
+				}
+			}
+		}
+	}
+	return false;
+}
+
 void FTraceController::UpdateStatus(const FTraceControlStatus& Message, FTraceStatus& Status)
 {
 	Status.TraceSystemStatus = static_cast<FTraceStatus::ETraceSystemStatus>(Message.TraceSystemStatus);
diff --git a/Engine/Source/Runtime/SessionServices/Private/TraceController.h b/Engine/Source/Runtime/SessionServices/Private/TraceController.h
--- a/Engine/Source/Runtime/SessionServices/Private/TraceController.h
+++ b/Engine/Source/Runtime/SessionServices/Private/TraceController.h
@@ -31,6 +31,7 @@ private:
 	virtual void SendSettingsUpdateRequest() override;
 	virtual bool HasAvailableInstance(const FGuid& InstanceId) override;
 	virtual bool IsActive(FGuid InstanceId, double TimeoutSec = 5.0f) const override;
+	virtual bool HasCapability(FGuid InstanceId, FStringView Name) override;
 	virtual void WithInstance(FGuid InstanceId, FCallback Func) override;
 	
 	DECLARE_DERIVED_EVENT(FTraceController, ITraceController::FStatusReceivedEvent, FStatusReceivedEvent);
@@ -86,7 +87,7 @@ private:
 	FSecureResponseEvent SecureResponse;
 
 	/** Lock to protect access to Instances list */
-	FRWLock InstancesLock;
+	mutable FRWLock InstancesLock;
 
 	/** Known instances with an active trace service */
 	TMap<FMessageAddress, FTracingInstance> Instances;
diff --git a/Engine/Source/Runtime/SessionServices/Public/ITraceController.h b/Engine/Source/Runtime/SessionServices/Public/ITraceController.h
--- a/Engine/Source/Runtime/SessionServices/Public/ITraceController.h
+++ b/Engine/Source/Runtime/SessionServices/Public/ITraceController.h
@@ -52,6 +52,11 @@ struct FTraceStatus
 		FString ChannelList;
 		bool bIsReadOnly;
 	};
+
+	struct FCapability
+	{
+		FString Name;
+	};
 	
 	struct FSettings
 	{
@@ -63,6 +68,8 @@ struct FTraceStatus
 		uint32 TailSizeBytes;
 		/** The channel presets defined by the process. */
 		TArray<FChannelPreset> ChannelPresets;
+		/** Capabilities of the target */
+		TArray<FCapability> Capabilities;
 	};
 
 	struct FChannel
@@ -244,6 +251,15 @@ public:
 	 */
 	virtual void WithInstance(FGuid InstanceId, FCallback Func) = 0;
 
+
+	/**
+	 * Queries a session about a specific capability.
+	 * @param InstanceId InstanceId to query
+	 * @param Name Name of capability
+	 * @return True if session lists capability, false otherwise
+	 */
+	virtual bool HasCapability(FGuid InstanceId, FStringView Name) = 0;
+
 	/**
 	 * Create an instance using a custom message bus
 	 * @param MessageBus Message bus the controller will use
